Cinema4D: Add support for Mac platform

#rnx

diff --git a/Engine/Plugins/Enterprise/DatasmithC4DImporter/DatasmithC4DImporter.uplugin b/Engine/Plugins/Enterprise/DatasmithC4DImporter/DatasmithC4DImporter.uplugin
--- a/Engine/Plugins/Enterprise/DatasmithC4DImporter/DatasmithC4DImporter.uplugin
+++ b/Engine/Plugins/Enterprise/DatasmithC4DImporter/DatasmithC4DImporter.uplugin
@@ -16,7 +16,8 @@
 	"Installed" : false,
 	"SupportedTargetPlatforms":
 	[
-		"Win64"
+		"Win64",
+		"Mac"
 	],
 	"Modules" :
 	[
@@ -26,7 +27,8 @@
 			"LoadingPhase" : "PostDefault",
 			"WhitelistPlatforms" :
 			[
-				"Win64"
+				"Win64",
+				"Mac"
 			]
 		}
 	],
diff --git a/Engine/Plugins/Enterprise/DatasmithC4DImporter/Source/DatasmithC4DTranslator/Private/DatasmithC4DMelangeSDKEnterGuard.h b/Engine/Plugins/Enterprise/DatasmithC4DImporter/Source/DatasmithC4DTranslator/Private/DatasmithC4DMelangeSDKEnterGuard.h
--- a/Engine/Plugins/Enterprise/DatasmithC4DImporter/Source/DatasmithC4DTranslator/Private/DatasmithC4DMelangeSDKEnterGuard.h
+++ b/Engine/Plugins/Enterprise/DatasmithC4DImporter/Source/DatasmithC4DTranslator/Private/DatasmithC4DMelangeSDKEnterGuard.h
@@ -1,14 +1,13 @@
 // Copyright 1998-2019 Epic Games, Inc. All Rights Reserved.
 
 // Enter Datasmith platform include gard.
+#if PLATFORM_WINDOWS
 #include "Windows/AllowWindowsPlatformTypes.h"
 
 // Back up some preprocessor identifiers.
 __pragma(push_macro("PI"))
 __pragma(push_macro("BYTE_MAX"))
 
-// Make Melange c4d_system.h build.
-#define __LP64__ 1
 #undef PI
 
 // Make Melange private_symbols.h build.
@@ -16,3 +15,8 @@ __pragma(push_macro("BYTE_MAX"))
 
 __pragma(warning(push))
 __pragma(warning(disable: 6297)) /* melange\20.004_rbmelange20.0_259890\includes\c4d_drawport.h(276): Arithmetic overflow:  32-bit value is shifted, then cast to 64-bit value.  Results might not be an expected value. */
+
+#endif // PLATFORM_WINDOWS
+
+
+
diff --git a/Engine/Plugins/Enterprise/DatasmithC4DImporter/Source/DatasmithC4DTranslator/Private/DatasmithC4DMelangeSDKLeaveGuard.h b/Engine/Plugins/Enterprise/DatasmithC4DImporter/Source/DatasmithC4DTranslator/Private/DatasmithC4DMelangeSDKLeaveGuard.h
--- a/Engine/Plugins/Enterprise/DatasmithC4DImporter/Source/DatasmithC4DTranslator/Private/DatasmithC4DMelangeSDKLeaveGuard.h
+++ b/Engine/Plugins/Enterprise/DatasmithC4DImporter/Source/DatasmithC4DTranslator/Private/DatasmithC4DMelangeSDKLeaveGuard.h
@@ -1,5 +1,7 @@
 // Copyright 1998-2019 Epic Games, Inc. All Rights Reserved.
 
+#if PLATFORM_WINDOWS
+
 __pragma(warning(pop))
 
 // Restore some preprocessor identifiers.
@@ -8,3 +10,4 @@ __pragma(pop_macro("BYTE_MAX"))
 
 // Leave Datasmith platform include guard.
 #include "Windows/HideWindowsPlatformTypes.h"
+#endif
