Copy up //UE4/Dev-Enterprise/... to //UE4/Dev-Tools-Staging/...

diff --git a/Engine/Plugins/Editor/GLTFImporter/GLTFImporter.uplugin b/Engine/Plugins/Editor/GLTFImporter/GLTFImporter.uplugin
--- a/Engine/Plugins/Editor/GLTFImporter/GLTFImporter.uplugin
+++ b/Engine/Plugins/Editor/GLTFImporter/GLTFImporter.uplugin
@@ -8,6 +8,7 @@
 	"CreatedBy": "Epic Games, Inc.",
 	"CreatedByURL": "http://epicgames.com",
 	"EnabledByDefault": false,
+	"CanContainContent": true,
 	"IsBetaVersion": true,
 	"Modules":
 	[
diff --git a/Engine/Plugins/Editor/GLTFImporter/Source/GLTFCore/Private/GLTFMapFactory.cpp b/Engine/Plugins/Editor/GLTFImporter/Source/GLTFCore/Private/GLTFMapFactory.cpp
--- a/Engine/Plugins/Editor/GLTFImporter/Source/GLTFCore/Private/GLTFMapFactory.cpp
+++ b/Engine/Plugins/Editor/GLTFImporter/Source/GLTFCore/Private/GLTFMapFactory.cpp
@@ -60,7 +60,9 @@ namespace GLTF
 
 		FMaterialExpressionTexture* TexExpression = CreateTextureMap(Map, CoordinateIndex, TEXT("Normal Map"), ETextureMode::Normal);
 		if (!TexExpression)
+		{
 			return;
+		}
 
 		FMaterialExpressionScalar* ScalarExpression = CurrentMaterialElement->AddMaterialExpression<FMaterialExpressionScalar>();
 		ScalarExpression->SetName(TEXT("Normal Scale"));
@@ -113,12 +115,6 @@ namespace GLTF
 		return CreateMap<FMaterialExpressionColor>(Map, CoordinateIndex, Color, MapName, ValueName, TextureMode, MaterialInput);
 	}
 
-	FMaterialExpression* FPBRMapFactory::CreateScalarMap(const GLTF::FTexture& Map, int CoordinateIndex, float Value, const TCHAR* MapName,
-	                                                     const TCHAR* ValueName, ETextureMode TextureMode, FMaterialExpressionInput& MaterialInput)
-	{
-		return CreateMap<FMaterialExpressionScalar>(Map, CoordinateIndex, Value, MapName, ValueName, TextureMode, MaterialInput);
-	}
-
 	void FPBRMapFactory::CreateMultiMap(const GLTF::FTexture& Map, int CoordinateIndex, const TCHAR* MapName, const FMapChannel* MapChannels,
 	                                    uint32 MapChannelsCount, ETextureMode TextureMode)
 	{
@@ -132,18 +128,20 @@ namespace GLTF
 			{
 				case EChannel::RG:
 				case EChannel::RGB:
+				case EChannel::All:
 				{
-					ValueExpression                           = CurrentMaterialElement->AddMaterialExpression<FMaterialExpressionColor>();
+					ValueExpression = CurrentMaterialElement->AddMaterialExpression<FMaterialExpressionColor>();
 					FMaterialExpressionColor* ColorExpression = static_cast<FMaterialExpressionColor*>(ValueExpression);
 					SetExpresisonValue(*reinterpret_cast<const FVector*>(MapChannel.VecValue), ColorExpression);
 					ColorExpression->SetGroupName(*GroupName);
 					break;
 				}
 				default:
-					ValueExpression                             = CurrentMaterialElement->AddMaterialExpression<FMaterialExpressionScalar>();
+					ValueExpression = CurrentMaterialElement->AddMaterialExpression<FMaterialExpressionScalar>();
 					FMaterialExpressionScalar* ScalarExpression = static_cast<FMaterialExpressionScalar*>(ValueExpression);
 					SetExpresisonValue(MapChannel.Value, ScalarExpression);
 					ScalarExpression->SetGroupName(*GroupName);
+					break;
 			}
 
 			ValueExpression->SetName(MapChannel.ValueName);
@@ -169,21 +167,17 @@ namespace GLTF
 					MakeFloat2->SetFunctionPathName(TEXT("/Engine/Functions/Engine_MaterialFunctions02/Utility/MakeFloat2.MakeFloat2"));
 
 					TexExpression->ConnectExpression(*MakeFloat2->GetInput(0), (int32)EChannel::Red);
-					TexExpression->ConnectExpression(*MakeFloat2->GetInput(1), (int32)EChannel::Greeen);
+					TexExpression->ConnectExpression(*MakeFloat2->GetInput(1), (int32)EChannel::Green);
 					MakeFloat2->ConnectExpression(*MultiplyExpression->GetInput(0), 0);
 					break;
 				}
+				// RGB is the top channel now
 				case EChannel::RGB:
-				{
-					FMaterialExpressionFunctionCall* MakeFloat3 = CurrentMaterialElement->AddMaterialExpression<FMaterialExpressionFunctionCall>();
-					MakeFloat3->SetFunctionPathName(TEXT("/Engine/Functions/Engine_MaterialFunctions02/Utility/MakeFloat3.MakeFloat3"));
-
-					TexExpression->ConnectExpression(*MakeFloat3->GetInput(0), (int32)EChannel::Red);
-					TexExpression->ConnectExpression(*MakeFloat3->GetInput(1), (int32)EChannel::Greeen);
-					TexExpression->ConnectExpression(*MakeFloat3->GetInput(2), (int32)EChannel::Blue);
-					MakeFloat3->ConnectExpression(*MultiplyExpression->GetInput(0), 0);
+					TexExpression->ConnectExpression(*MultiplyExpression->GetInput(0), 0);
+					break;
+				case EChannel::All:
+					TexExpression->ConnectExpression(*MultiplyExpression->GetInput(0), 5);
 					break;
-				}
 				default:
 					// single channel connection
 					TexExpression->ConnectExpression(*MultiplyExpression->GetInput(0), (int32)MapChannel.Channel);
@@ -197,7 +191,9 @@ namespace GLTF
 					MapChannel.OutputExpression->ConnectExpression(*MapChannel.MaterialInput, 0);
 				}
 				else
+				{
 					MultiplyExpression->ConnectExpression(*MapChannel.MaterialInput, 0);
+				}
 			}
 		}
 		else
@@ -213,13 +209,13 @@ namespace GLTF
 					MapChannel.OutputExpression->ConnectExpression(*MapChannel.MaterialInput, 0);
 				}
 				else
+				{
 					ValueExpressions[Index]->ConnectExpression(*MapChannel.MaterialInput, 0);
+				}
 			}
 		}
 	}
 
-	
-
 	FMaterialExpressionTexture* FPBRMapFactory::CreateTextureMap(const GLTF::FTexture& Map, int CoordinateIndex, const TCHAR* MapName,
 	                                                             ETextureMode TextureMode)
 	{
diff --git a/Engine/Plugins/Editor/GLTFImporter/Source/GLTFCore/Private/GLTFMapFactory.h b/Engine/Plugins/Editor/GLTFImporter/Source/GLTFCore/Private/GLTFMapFactory.h
--- a/Engine/Plugins/Editor/GLTFImporter/Source/GLTFCore/Private/GLTFMapFactory.h
+++ b/Engine/Plugins/Editor/GLTFImporter/Source/GLTFCore/Private/GLTFMapFactory.h
@@ -18,7 +18,7 @@ namespace GLTF
 		{
 			All,
 			Red,
-			Greeen,
+			Green,
 			Blue,
 			Alpha,
 			RG,
@@ -43,10 +43,10 @@ namespace GLTF
 				Value = InValue;
 			}
 
-			FMapChannel(float InVecValue[3], const TCHAR* InValueName, EChannel InChannel, FMaterialExpressionInput* InMaterialInput, FMaterialExpression* InOutputExpression)
+			FMapChannel(const FVector& InVecValue, const TCHAR* InValueName, EChannel InChannel, FMaterialExpressionInput* InMaterialInput, FMaterialExpression* InOutputExpression)
 				: FMapChannel(InValueName, InChannel, InMaterialInput, InOutputExpression)
 			{
-				SetValue( FVector(InVecValue[0], InVecValue[1], InVecValue[2]) );
+				SetValue(InVecValue);
 			}
 
 			void SetValue(const FVector& Vec)
@@ -80,9 +80,6 @@ namespace GLTF
 		FMaterialExpression* CreateColorMap(const GLTF::FTexture& Map, int CoordinateIndex, const FVector4& Color, const TCHAR* MapName,
 		                                    const TCHAR* ValueName, ETextureMode TextureMode, FMaterialExpressionInput& MaterialInput);
 
-		FMaterialExpression* CreateScalarMap(const GLTF::FTexture& Map, int CoordinateIndex, float Value, const TCHAR* MapName,
-		                                     const TCHAR* ValueName, ETextureMode TextureMode, FMaterialExpressionInput& MaterialInput);
-
 		FMaterialExpressionTexture* CreateTextureMap(const GLTF::FTexture& Map, int CoordinateIndex, const TCHAR* MapName, ETextureMode TextureMode);
 
 		void CreateMultiMap(const GLTF::FTexture& Map, int CoordinateIndex, const TCHAR* MapName, const FMapChannel* MapChannels,
diff --git a/Engine/Plugins/Editor/GLTFImporter/Source/GLTFCore/Private/GLTFMaterialFactory.cpp b/Engine/Plugins/Editor/GLTFImporter/Source/GLTFCore/Private/GLTFMaterialFactory.cpp
--- a/Engine/Plugins/Editor/GLTFImporter/Source/GLTFCore/Private/GLTFMaterialFactory.cpp
+++ b/Engine/Plugins/Editor/GLTFImporter/Source/GLTFCore/Private/GLTFMaterialFactory.cpp
@@ -34,13 +34,12 @@ namespace GLTF
 		const TArray<FMaterialElement*>& CreateMaterials(const GLTF::FAsset& Asset, UObject* ParentPackage, EObjectFlags Flags);
 
 	private:
-		void HandleOcclusion(const TArray<GLTF::FTexture>& Texture, const GLTF::FMaterial& GLTFMaterial, FPBRMapFactory& MapFactory,
-		                     FMaterialElement& MaterialElement);
-
-		void HandleGGX(const TArray<GLTF::FTexture>& Texture, const GLTF::FMaterial& GLTFMaterial, FPBRMapFactory& MapFactory,
-		               FMaterialElement& MaterialElement);
-
 		void HandleOpacity(const TArray<GLTF::FTexture>& Texture, const GLTF::FMaterial& GLTFMaterial, FMaterialElement& MaterialElement);
+		void HandleShadingModel(const TArray<GLTF::FTexture>& Texture, const GLTF::FMaterial& GLTFMaterial, FPBRMapFactory& MapFactory, FMaterialElement& MaterialElement);
+		void HandleOcclusion(const TArray<GLTF::FTexture>& Texture, const GLTF::FMaterial& GLTFMaterial, FPBRMapFactory& MapFactory, FMaterialElement& MaterialElement);
+		void HandleEmissive(const TArray<GLTF::FTexture>& Textures, const GLTF::FMaterial& GLTFMaterial, FPBRMapFactory& MapFactory, FMaterialElement& MaterialElement);
+		void HandleNormal(const TArray<GLTF::FTexture>& Textures, const GLTF::FMaterial& GLTFMaterial, FPBRMapFactory& MapFactory, FMaterialElement& MaterialElement);
+
 
 	private:
 		TUniquePtr<IMaterialElementFactory> MaterialElementFactory;
@@ -112,38 +111,177 @@ namespace GLTF
 
 			MapFactory.CurrentMaterialElement = MaterialElement;
 
-			MapFactory.GroupName = TEXT("Base Color");
-			MapFactory.CreateColorMap(GetTexture(GLTFMaterial.BaseColor, Asset.Textures), GLTFMaterial.BaseColor.TexCoord,
-			                          GLTFMaterial.BaseColorFactor, TEXT("BaseColor"), nullptr, ETextureMode::Color, MaterialElement->GetBaseColor());
-
-			MapFactory.GroupName = TEXT("Normal");
-			MapFactory.CreateNormalMap(GetTexture(GLTFMaterial.Normal, Asset.Textures), GLTFMaterial.Normal.TexCoord, GLTFMaterial.NormalScale);
-			if (GLTFMaterial.Emissive.TextureIndex != INDEX_NONE && !GLTFMaterial.EmissiveFactor.IsNearlyZero())
-			{
-				MapFactory.GroupName = TEXT("Emission");
-				MapFactory.CreateColorMap(GetTexture(GLTFMaterial.Emissive, Asset.Textures), GLTFMaterial.Emissive.TexCoord,
-				                          GLTFMaterial.EmissiveFactor, TEXT("Emissive"), TEXT("Color"), ETextureMode::Color,
-				                          MaterialElement->GetEmissiveColor());  // emissive map is in sRGB space
-			}
+			HandleShadingModel(Asset.Textures, GLTFMaterial, MapFactory, *MaterialElement);
+			HandleOpacity(Asset.Textures, GLTFMaterial, *MaterialElement);
 
+			// Additional maps
 			HandleOcclusion(Asset.Textures, GLTFMaterial, MapFactory, *MaterialElement);
-			HandleGGX(Asset.Textures, GLTFMaterial, MapFactory, *MaterialElement);
-			HandleOpacity(Asset.Textures, GLTFMaterial, *MaterialElement);
+			HandleEmissive(Asset.Textures, GLTFMaterial, MapFactory, *MaterialElement);
+			HandleNormal(Asset.Textures, GLTFMaterial, MapFactory, *MaterialElement);
 
-			MaterialElement->Finalize();
+			MaterialElement->SetGLTFMaterialHash(GLTFMaterial.GetHash());
+            MaterialElement->Finalize();
 			Materials.Add(MaterialElement);
 		}
 
 		return Materials;
 	}
 
-	void FMaterialFactoryImpl::HandleOcclusion(const TArray<GLTF::FTexture>& Textures, const GLTF::FMaterial& GLTFMaterial,
-	                                           FPBRMapFactory& MapFactory, FMaterialElement& MaterialElement)
+	void FMaterialFactoryImpl::HandleOpacity(const TArray<GLTF::FTexture>& Textures, const GLTF::FMaterial& GLTFMaterial, FMaterialElement& MaterialElement)
+	{
+		if (GLTFMaterial.IsOpaque())
+		{
+			return;
+		}
+
+		const TCHAR* GroupName = TEXT("Opacity");
+
+		FMaterialExpressionTexture* BaseColorMap = FindExpression<FMaterialExpressionTexture>(TEXT("BaseColor Map"), MaterialElement);
+		switch (GLTFMaterial.AlphaMode)
+		{
+		case FMaterial::EAlphaMode::Mask:
+		{
+			FMaterialExpressionColor* BaseColorFactor = FindExpression<FMaterialExpressionColor>(TEXT("BaseColor"), MaterialElement);
+
+			FMaterialExpressionGeneric* MultiplyExpression = MaterialElement.AddMaterialExpression<FMaterialExpressionGeneric>();
+			MultiplyExpression->SetExpressionName(TEXT("Multiply"));
+			BaseColorFactor->ConnectExpression(*MultiplyExpression->GetInput(1), (int)FPBRMapFactory::EChannel::Alpha);
+			BaseColorMap->ConnectExpression(*MultiplyExpression->GetInput(0), (int)FPBRMapFactory::EChannel::Alpha);
+
+			FMaterialExpressionFunctionCall* CuttofExpression = MaterialElement.AddMaterialExpression<FMaterialExpressionFunctionCall>();
+			CuttofExpression->SetFunctionPathName(TEXT("/Engine/Functions/Engine_MaterialFunctions02/SmoothStep.SmoothStep"));
+
+			FMaterialExpressionScalar* ValueExpression = MaterialElement.AddMaterialExpression<FMaterialExpressionScalar>();
+			ValueExpression->SetName(TEXT("Alpha Cuttof"));
+			ValueExpression->SetGroupName(GroupName);
+			ValueExpression->GetScalar() = GLTFMaterial.AlphaCutoff;
+
+			MultiplyExpression->ConnectExpression(*CuttofExpression->GetInput(0), 0);
+			ValueExpression->ConnectExpression(*CuttofExpression->GetInput(1), 0);
+			ValueExpression->ConnectExpression(*CuttofExpression->GetInput(2), 0);
+
+			CuttofExpression->ConnectExpression(MaterialElement.GetOpacity(), 0);
+			break;
+		}
+		case FMaterial::EAlphaMode::Blend:
+		{
+			FMaterialExpressionScalar* ValueExpression = MaterialElement.AddMaterialExpression<FMaterialExpressionScalar>();
+			ValueExpression->SetName(TEXT("IOR"));
+			ValueExpression->SetGroupName(GroupName);
+			ValueExpression->GetScalar() = 1.f;
+			ValueExpression->ConnectExpression(MaterialElement.GetRefraction(), 0);
+
+			FMaterialExpressionColor* BaseColorFactor = FindExpression<FMaterialExpressionColor>(TEXT("BaseColor"), MaterialElement);
+			if (BaseColorMap)
+			{
+				FMaterialExpressionGeneric* MultiplyExpression = MaterialElement.AddMaterialExpression<FMaterialExpressionGeneric>();
+				MultiplyExpression->SetExpressionName(TEXT("Multiply"));
+				BaseColorFactor->ConnectExpression(*MultiplyExpression->GetInput(1), (int)FPBRMapFactory::EChannel::Alpha);
+				BaseColorMap->ConnectExpression(*MultiplyExpression->GetInput(0), (int)FPBRMapFactory::EChannel::Alpha);
+				MultiplyExpression->ConnectExpression(MaterialElement.GetOpacity(), 0);
+			}
+			else
+			{
+				BaseColorFactor->ConnectExpression(MaterialElement.GetOpacity(), (int)FPBRMapFactory::EChannel::Alpha);
+			}
+			break;
+		}
+		default:
+			check(false);
+			break;
+		}
+	}
+
+	void FMaterialFactoryImpl::HandleShadingModel(const TArray<GLTF::FTexture>& Textures, const GLTF::FMaterial& GLTFMaterial, FPBRMapFactory& MapFactory, FMaterialElement& MaterialElement)
+	{
+		MapFactory.GroupName = TEXT("GGX");
+
+		TArray<FPBRMapFactory::FMapChannel, TFixedAllocator<4>> Maps;
+		if (GLTFMaterial.ShadingModel == FMaterial::EShadingModel::MetallicRoughness)
+		{
+			// Base Color
+			MapFactory.GroupName = TEXT("Base Color");
+			MapFactory.CreateColorMap(GetTexture(GLTFMaterial.BaseColor, Textures),
+									  GLTFMaterial.BaseColor.TexCoord,
+									  GLTFMaterial.BaseColorFactor,
+									  TEXT("BaseColor"),
+									  nullptr,
+									  ETextureMode::Color,
+									  MaterialElement.GetBaseColor());
+
+			// Metallic
+			Maps.Emplace(GLTFMaterial.MetallicRoughness.MetallicFactor,
+						 TEXT("Metallic Factor"),
+						 FPBRMapFactory::EChannel::Blue,
+						 &MaterialElement.GetMetallic(),
+						 nullptr);
+
+			// Roughness
+			Maps.Emplace(GLTFMaterial.MetallicRoughness.RoughnessFactor,
+						 TEXT("Roughness Factor"),
+						 FPBRMapFactory::EChannel::Green,
+						 &MaterialElement.GetRoughness(),
+						 nullptr);
+
+			MapFactory.CreateMultiMap(GetTexture(GLTFMaterial.MetallicRoughness.Map, Textures),
+									  GLTFMaterial.MetallicRoughness.Map.TexCoord,
+				                      TEXT("MetallicRoughness"),
+									  Maps.GetData(),
+									  Maps.Num(),
+									  ETextureMode::Grayscale);
+		}
+		else if (GLTFMaterial.ShadingModel == FMaterial::EShadingModel::SpecularGlossiness)
+		{
+			// We'll actually just convert it into MetalRoughness in the material graph
+			FMaterialExpressionFunctionCall* SpecGlossToMetalRough = MaterialElement.AddMaterialExpression<FMaterialExpressionFunctionCall>();
+			SpecGlossToMetalRough->SetFunctionPathName(TEXT("/GLTFImporter/SpecGlossToMetalRoughness.SpecGlossToMetalRoughness"));
+			SpecGlossToMetalRough->ConnectExpression(MaterialElement.GetBaseColor(), 0);
+			SpecGlossToMetalRough->ConnectExpression(MaterialElement.GetMetallic(), 1);
+
+			FMaterialExpressionGeneric* GlossToRoughness = MaterialElement.AddMaterialExpression<FMaterialExpressionGeneric>();
+			GlossToRoughness->SetExpressionName(TEXT("OneMinus"));
+			GlossToRoughness->ConnectExpression(MaterialElement.GetRoughness(), 0);
+
+			// Diffuse Color (BaseColor/BaseColorFactor are used to store the Diffuse alternatives for Spec/Gloss)
+			MapFactory.GroupName = TEXT("Diffuse Color");
+			FMaterialExpression* Diffuse = MapFactory.CreateColorMap(GetTexture(GLTFMaterial.BaseColor, Textures),
+																	 GLTFMaterial.BaseColor.TexCoord,
+																	 GLTFMaterial.BaseColorFactor,
+																	 TEXT("Diffuse"),
+																	 TEXT("Color"),
+																	 ETextureMode::Color,
+																	 *SpecGlossToMetalRough->GetInput(1));
+
+			// Specular (goes into SpecGlossToMetalRough conversion)
+			Maps.Emplace(GLTFMaterial.SpecularGlossiness.SpecularFactor,
+						 TEXT("Specular Factor"),
+						 FPBRMapFactory::EChannel::RGB,
+						 SpecGlossToMetalRough->GetInput(0),
+						 nullptr);
+
+			// Glossiness (converted to Roughness)
+			Maps.Emplace(GLTFMaterial.SpecularGlossiness.GlossinessFactor,
+						 TEXT("Glossiness Factor"),
+				         FPBRMapFactory::EChannel::Alpha,
+						 GlossToRoughness->GetInput(0),
+						 nullptr);
+
+			// Creates the multimap for Specular and Glossiness
+			MapFactory.CreateMultiMap(GetTexture(GLTFMaterial.SpecularGlossiness.Map, Textures),
+									  GLTFMaterial.SpecularGlossiness.Map.TexCoord,
+				                      TEXT("SpecularGlossiness"),
+									  Maps.GetData(),
+									  Maps.Num(),
+				                      ETextureMode::Color);
+		}
+	}
+
+	void FMaterialFactoryImpl::HandleOcclusion(const TArray<GLTF::FTexture>& Textures, const GLTF::FMaterial& GLTFMaterial, FPBRMapFactory& MapFactory, FMaterialElement& MaterialElement)
 	{
 		MapFactory.GroupName = TEXT("Occlusion");
 
 		FMaterialExpressionTexture* TexExpression = MapFactory.CreateTextureMap(
-		    GetTexture(GLTFMaterial.Occlusion, Textures), GLTFMaterial.Occlusion.TexCoord, TEXT("Occlusion"), ETextureMode::Grayscale);
+			GetTexture(GLTFMaterial.Occlusion, Textures), GLTFMaterial.Occlusion.TexCoord, TEXT("Occlusion"), ETextureMode::Grayscale);
 
 		if (!TexExpression)
 			return;
@@ -166,168 +304,31 @@ namespace GLTF
 		LerpExpression->ConnectExpression(MaterialElement.GetAmbientOcclusion(), 0);
 	}
 
-	void FMaterialFactoryImpl::HandleGGX(const TArray<GLTF::FTexture>& Textures, const GLTF::FMaterial& GLTFMaterial, FPBRMapFactory& MapFactory,
-	                                     FMaterialElement& MaterialElement)
+	void FMaterialFactoryImpl::HandleEmissive(const TArray<GLTF::FTexture>& Textures, const GLTF::FMaterial& GLTFMaterial, FPBRMapFactory& MapFactory, FMaterialElement& MaterialElement)
 	{
-		MapFactory.GroupName = TEXT("GGX");
-
-		// will need to correct the roughness which is GGX
-		FMaterialExpressionGeneric* SqrtExpression = MaterialElement.AddMaterialExpression<FMaterialExpressionGeneric>();
-		SqrtExpression->SetExpressionName(TEXT("SquareRoot"));
-
-		TArray<FPBRMapFactory::FMapChannel, TFixedAllocator<4> > Maps;
-		switch (GLTFMaterial.ShadingModel)
+		if (GLTFMaterial.Emissive.TextureIndex == INDEX_NONE || GLTFMaterial.EmissiveFactor.IsNearlyZero())
 		{
-			case FMaterial::EShadingModel::MetallicRoughness:
-			{
-				// according to the GLTF specs:
-				// cdiff = lerp(baseColor.rgb * (1 - dielectricSpecular.r), black, metallic)
-				// F0 = lerp(dieletricSpecular, baseColor.rgb, metallic)
-				// alpha = roughness ^ 2
-
-				Maps.Emplace(GLTFMaterial.MetallicRoughness.MetallicFactor, TEXT("Metallic Factor"),
-				                                      FPBRMapFactory::EChannel::Blue, &MaterialElement.GetMetallic(), nullptr);
-				Maps.Emplace(GLTFMaterial.MetallicRoughness.RoughnessFactor, TEXT("Roughness Factor"),
-				                                      FPBRMapFactory::EChannel::Greeen, &MaterialElement.GetRoughness(), SqrtExpression);
-
-				MapFactory.CreateMultiMap(GetTexture(GLTFMaterial.MetallicRoughness.Map, Textures), GLTFMaterial.MetallicRoughness.Map.TexCoord,
-				                          TEXT("MetallicRoughness Map"), Maps.GetData(), Maps.Num(), ETextureMode::Grayscale);
-
-				// GLTF specifies that the dielectricSpecular is, 0.04 while for UE4 it's 0.08 * Specular, so correct it
-				FMaterialExpressionScalar* ValueExpression = MaterialElement.AddMaterialExpression<FMaterialExpressionScalar>();
-				ValueExpression->GetScalar()               = 0.5f;
-				ValueExpression->ConnectExpression(MaterialElement.GetSpecular(), 0);
-
-				break;
-			}
-			case FMaterial::EShadingModel::SpecularGlossiness:
-			{
-				// according to the GLTF specs:
-				// cdiff = diffuse.rgb * (1 - max(specular.r, specular.g, specular.b))
-				// F0 = specular
-				// alpha = (1 - glossiness) ^ 2
-
-				// convert glossiness to roughness
-				{
-					FMaterialExpressionGeneric* NegExpression = MaterialElement.AddMaterialExpression<FMaterialExpressionGeneric>();
-					NegExpression->SetExpressionName(TEXT("OneMinus"));
-					NegExpression->ConnectExpression(*SqrtExpression->GetInput(0), 0);
-				}
-
-				// create the multi map expressions
-				Maps.Emplace(0.f, TEXT("Specular Factor"), FPBRMapFactory::EChannel::RGB, &MaterialElement.GetSpecular(),
-				                                      nullptr);
-				Maps[0].SetValue(GLTFMaterial.SpecularGlossiness.SpecularFactor);
-				Maps.Emplace(GLTFMaterial.SpecularGlossiness.GlossinessFactor, TEXT("Glossiness Factor"),
-				                                      FPBRMapFactory::EChannel::Alpha, &MaterialElement.GetRoughness(), SqrtExpression);
-
-				MapFactory.CreateMultiMap(GetTexture(GLTFMaterial.SpecularGlossiness.Map, Textures), GLTFMaterial.SpecularGlossiness.Map.TexCoord,
-				                          TEXT("SpecularGlossiness Map"), Maps.GetData(), Maps.Num(),
-				                          ETextureMode::Color);  // specular map is in sRGB space
-
-				// adjust diffuse with specular
-				{
-					FMaterialExpressionColor* BaseColorFactor = FindExpression<FMaterialExpressionColor>(TEXT("BaseColor"), MaterialElement);
-					check(BaseColorFactor);
-					const float SpecValue = 1.0 - GLTFMaterial.SpecularGlossiness.SpecularFactor.GetMax();
-					BaseColorFactor->GetColor().R *= SpecValue;
-					BaseColorFactor->GetColor().G *= SpecValue;
-					BaseColorFactor->GetColor().B *= SpecValue;
-				}
-
-				FMaterialExpression* BaseColor = MaterialElement.GetBaseColor().GetExpression();
-				// convert specular to diffuse term
-				{
-					FMaterialExpression*        Specular      = MaterialElement.GetSpecular().GetExpression();
-					FMaterialExpressionGeneric* AddExpression = MaterialElement.AddMaterialExpression<FMaterialExpressionGeneric>();
-					AddExpression->SetExpressionName(TEXT("Add"));
-
-					BaseColor->ConnectExpression(*AddExpression->GetInput(0), 0);
-					Specular->ConnectExpression(*AddExpression->GetInput(1), 0);
-					AddExpression->ConnectExpression(MaterialElement.GetBaseColor(), 0);
-				}
-				// convert diffuse to metallic, i.e. when diffuse zero material is metallic
-				{
-					FMaterialExpressionGeneric* NegExpression =
-					    MaterialElement.AddMaterialExpression<FMaterialExpressionGeneric>();  // invert glossiness
-					NegExpression->SetExpressionName(TEXT("OneMinus"));
-					BaseColor->ConnectExpression(*NegExpression->GetInput(0), 0);
-					NegExpression->ConnectExpression(MaterialElement.GetMetallic(), 0);
-				}
-
-				break;
-			}
-			default:
-				check(false);
-				break;
+			return;
 		}
+
+		MapFactory.GroupName = TEXT("Emission");
+		MapFactory.CreateColorMap(GetTexture(GLTFMaterial.Emissive, Textures),
+								  GLTFMaterial.Emissive.TexCoord,
+								  GLTFMaterial.EmissiveFactor,
+								  TEXT("Emissive"),
+								  TEXT("Color"),
+								  ETextureMode::Color, // emissive map is in sRGB space
+								  MaterialElement.GetEmissiveColor());
 	}
 
-	void FMaterialFactoryImpl::HandleOpacity(const TArray<GLTF::FTexture>& Textures, const GLTF::FMaterial& GLTFMaterial,
-	                                         FMaterialElement& MaterialElement)
+	void FMaterialFactoryImpl::HandleNormal(const TArray<GLTF::FTexture>& Textures, const GLTF::FMaterial& GLTFMaterial, FPBRMapFactory& MapFactory, FMaterialElement& MaterialElement)
 	{
-		if (GLTFMaterial.IsOpaque())
-			return;
-
-		const TCHAR* GroupName = TEXT("Opacity");
-
-		FMaterialExpressionTexture* BaseColorMap = FindExpression<FMaterialExpressionTexture>(TEXT("BaseColor Map"), MaterialElement);
-		switch (GLTFMaterial.AlphaMode)
-		{
-			case FMaterial::EAlphaMode::Mask:
-			{
-				FMaterialExpressionColor* BaseColorFactor = FindExpression<FMaterialExpressionColor>(TEXT("BaseColor"), MaterialElement);
-
-				FMaterialExpressionGeneric* MultiplyExpression = MaterialElement.AddMaterialExpression<FMaterialExpressionGeneric>();
-				MultiplyExpression->SetExpressionName(TEXT("Multiply"));
-				BaseColorFactor->ConnectExpression(*MultiplyExpression->GetInput(1), (int)FPBRMapFactory::EChannel::Alpha);
-				BaseColorMap->ConnectExpression(*MultiplyExpression->GetInput(0), (int)FPBRMapFactory::EChannel::Alpha);
-
-				FMaterialExpressionFunctionCall* CuttofExpression = MaterialElement.AddMaterialExpression<FMaterialExpressionFunctionCall>();
-				CuttofExpression->SetFunctionPathName(TEXT("/Engine/Functions/Engine_MaterialFunctions02/SmoothStep.SmoothStep"));
-
-				FMaterialExpressionScalar* ValueExpression = MaterialElement.AddMaterialExpression<FMaterialExpressionScalar>();
-				ValueExpression->SetName(TEXT("Alpha Cuttof"));
-				ValueExpression->SetGroupName(GroupName);
-				ValueExpression->GetScalar() = GLTFMaterial.AlphaCutoff;
-
-				MultiplyExpression->ConnectExpression(*CuttofExpression->GetInput(0), 0);
-				ValueExpression->ConnectExpression(*CuttofExpression->GetInput(1), 0);
-				ValueExpression->ConnectExpression(*CuttofExpression->GetInput(2), 0);
-
-				CuttofExpression->ConnectExpression(MaterialElement.GetOpacity(), 0);
-				break;
-			}
-			case FMaterial::EAlphaMode::Blend:
-			{
-				FMaterialExpressionScalar* ValueExpression = MaterialElement.AddMaterialExpression<FMaterialExpressionScalar>();
-				ValueExpression->SetName(TEXT("IOR"));
-				ValueExpression->SetGroupName(GroupName);
-				ValueExpression->GetScalar() = 1.f;
-				ValueExpression->ConnectExpression(MaterialElement.GetRefraction(), 0);
-
-				FMaterialExpressionColor* BaseColorFactor = FindExpression<FMaterialExpressionColor>(TEXT("BaseColor"), MaterialElement);
-				if (BaseColorMap)
-				{
-					FMaterialExpressionGeneric* MultiplyExpression = MaterialElement.AddMaterialExpression<FMaterialExpressionGeneric>();
-					MultiplyExpression->SetExpressionName(TEXT("Multiply"));
-					BaseColorFactor->ConnectExpression(*MultiplyExpression->GetInput(1), (int)FPBRMapFactory::EChannel::Alpha);
-					BaseColorMap->ConnectExpression(*MultiplyExpression->GetInput(0), (int)FPBRMapFactory::EChannel::Alpha);
-					MultiplyExpression->ConnectExpression(MaterialElement.GetOpacity(), 0);
-				}
-				else
-					BaseColorFactor->ConnectExpression(MaterialElement.GetOpacity(), (int)FPBRMapFactory::EChannel::Alpha);
-
-				break;
-			}
-			default:
-				check(false);
-				break;
-		}
+		MapFactory.GroupName = TEXT("Normal");
+		MapFactory.CreateNormalMap(GetTexture(GLTFMaterial.Normal, Textures),
+								   GLTFMaterial.Normal.TexCoord,
+								   GLTFMaterial.NormalScale);
 	}
 
-	//
-
 	FMaterialFactory::FMaterialFactory(IMaterialElementFactory* MaterialElementFactory, ITextureFactory* TextureFactory)
 	    : Impl(new FMaterialFactoryImpl(MaterialElementFactory, TextureFactory))
 	{
diff --git a/Engine/Plugins/Editor/GLTFImporter/Source/GLTFCore/Public/GLTFAccessor.h b/Engine/Plugins/Editor/GLTFImporter/Source/GLTFCore/Public/GLTFAccessor.h
--- a/Engine/Plugins/Editor/GLTFImporter/Source/GLTFCore/Public/GLTFAccessor.h
+++ b/Engine/Plugins/Editor/GLTFImporter/Source/GLTFCore/Public/GLTFAccessor.h
@@ -4,6 +4,8 @@
 
 #include "CoreMinimal.h"
 
+struct FMD5Hash;
+
 namespace GLTF
 {
 	struct GLTFCORE_API FBuffer
@@ -26,8 +28,8 @@ namespace GLTF
 	struct GLTFCORE_API FBufferView
 	{
 		const FBuffer& Buffer;
-		const uint32   ByteOffset;
-		const uint32   ByteLength;
+		const uint32 ByteOffset;
+		const uint32 ByteLength;
 		// if zero then accessor elements are tightly packed, i.e., effective stride equals the size of the element
 		const uint32 ByteStride;  // range 4..252
 
@@ -83,6 +85,7 @@ namespace GLTF
 		FAccessor(uint32 InCount, EType InType, EComponentType InComponentType, bool InNormalized);
 
 		virtual bool IsValid() const = 0;
+		virtual FMD5Hash GetHash() const = 0;
 
 		virtual uint32 GetUnsignedInt(uint32 Index) const;
 		virtual void   GetUnsignedInt16x4(uint32 Index, uint16 Values[4]) const;
@@ -118,6 +121,7 @@ namespace GLTF
 		FValidAccessor(FBufferView& InBufferView, uint32 InOffset, uint32 InCount, EType InType, EComponentType InCompType, bool InNormalized);
 
 		bool IsValid() const override;
+		FMD5Hash GetHash() const override;
 
 		uint32 GetUnsignedInt(uint32 Index) const override;
 		void   GetUnsignedInt16x4(uint32 Index, uint16 Values[4]) const override;
@@ -140,6 +144,7 @@ namespace GLTF
 		const FBufferView& BufferView;
 		const uint32       ByteOffset;
 		const uint32       ElementSize;
+		const uint32	   ByteStride;
 
 		const uint8* DataAt(uint32 Index) const;
 	};
@@ -152,6 +157,7 @@ namespace GLTF
 		}
 
 		bool IsValid() const override;
+		FMD5Hash GetHash() const override;
 	};
 
 	//
diff --git a/Engine/Plugins/Editor/GLTFImporter/Source/GLTFCore/Public/GLTFMaterial.h b/Engine/Plugins/Editor/GLTFImporter/Source/GLTFCore/Public/GLTFMaterial.h
--- a/Engine/Plugins/Editor/GLTFImporter/Source/GLTFCore/Public/GLTFMaterial.h
+++ b/Engine/Plugins/Editor/GLTFImporter/Source/GLTFCore/Public/GLTFMaterial.h
@@ -3,6 +3,7 @@
 #pragma once
 
 #include "CoreMinimal.h"
+#include "Misc/SecureHash.h"
 
 namespace GLTF
 {
@@ -16,6 +17,16 @@ namespace GLTF
 		    , TexCoord(0)
 		{
 		}
+
+		FMD5Hash GetHash() const
+		{
+			FMD5 MD5;
+			MD5.Update(reinterpret_cast<const uint8*>(&TextureIndex), sizeof(TextureIndex));
+			MD5.Update(&TexCoord, sizeof(TexCoord));
+			FMD5Hash Hash;
+			Hash.Set(MD5);
+			return Hash;
+		}
 	};
 
 	struct GLTFCORE_API FMaterial
@@ -87,8 +98,8 @@ namespace GLTF
 		FString Name;
 
 		// PBR properties
-		FTextureMap         BaseColor;
-		FVector4            BaseColorFactor;
+		FTextureMap         BaseColor;		 // Used for DiffuseColor on Specular-Glossiness mode
+		FVector4            BaseColorFactor; // Used for DiffuseFactor on Specular-Glossiness mode
 		EShadingModel       ShadingModel;
 		FMetallicRoughness  MetallicRoughness;
 		FSpecularGlossiness SpecularGlossiness;
@@ -128,5 +139,8 @@ namespace GLTF
 		{
 			return AlphaMode == EAlphaMode::Opaque;
 		}
+
+		FMD5Hash GetHash() const;
 	};
 }  // namespace GLTF
+
diff --git a/Engine/Plugins/Enterprise/DatasmithC4DImporter/Source/DatasmithC4DTranslator/DatasmithC4DTranslator.Build.cs b/Engine/Plugins/Enterprise/DatasmithC4DImporter/Source/DatasmithC4DTranslator/DatasmithC4DTranslator.Build.cs
--- a/Engine/Plugins/Enterprise/DatasmithC4DImporter/Source/DatasmithC4DTranslator/DatasmithC4DTranslator.Build.cs
+++ b/Engine/Plugins/Enterprise/DatasmithC4DImporter/Source/DatasmithC4DTranslator/DatasmithC4DTranslator.Build.cs
@@ -21,17 +21,14 @@ namespace UnrealBuildTool.Rules
 					"DatasmithExporter",
 					"Engine",
 					"Json",
-					"MainFrame",
 					"MeshDescription",
 					"MeshDescriptionOperations",
-					"MeshUtilities",
 					"MessageLog",
 					"Slate",
 					"SlateCore",
 					"StaticMeshDescription",
 					"UEOpenExr",
-					"UnrealEd",
-				}
+                }
 			);
 
 			PublicDependencyModuleNames.AddRange(
diff --git a/Engine/Plugins/Enterprise/DatasmithC4DImporter/Source/DatasmithC4DTranslator/Private/DatasmithC4DImportOptions.cpp b/Engine/Plugins/Enterprise/DatasmithC4DImporter/Source/DatasmithC4DTranslator/Private/DatasmithC4DImportOptions.cpp
--- a/Engine/Plugins/Enterprise/DatasmithC4DImporter/Source/DatasmithC4DTranslator/Private/DatasmithC4DImportOptions.cpp
+++ b/Engine/Plugins/Enterprise/DatasmithC4DImporter/Source/DatasmithC4DTranslator/Private/DatasmithC4DImportOptions.cpp
@@ -17,6 +17,7 @@ UDatasmithC4DImportOptions::UDatasmithC4DImportOptions(const FObjectInitializer&
 	ScaleVertices = 1.0;
 	bExportToUDatasmith = false;
 
+#if WITH_EDITOR
 	// In debug show all properties including the ones for debug
 	if (IDatasmithC4DTranslatorModule::Get().InDebugMode())
 	{
@@ -29,6 +30,7 @@ UDatasmithC4DImportOptions::UDatasmithC4DImportOptions(const FObjectInitializer&
 			}
 		}
 	}
+#endif //WITH_EDITOR
 }
 
 #undef LOCTEXT_NAMESPACE
diff --git a/Engine/Plugins/Enterprise/DatasmithC4DImporter/Source/DatasmithC4DTranslator/Private/DatasmithC4DImporter.cpp b/Engine/Plugins/Enterprise/DatasmithC4DImporter/Source/DatasmithC4DTranslator/Private/DatasmithC4DImporter.cpp
--- a/Engine/Plugins/Enterprise/DatasmithC4DImporter/Source/DatasmithC4DTranslator/Private/DatasmithC4DImporter.cpp
+++ b/Engine/Plugins/Enterprise/DatasmithC4DImporter/Source/DatasmithC4DTranslator/Private/DatasmithC4DImporter.cpp
@@ -6,14 +6,11 @@
 
 #include "DatasmithAssetImportData.h"
 #include "DatasmithC4DExtraMelangeDefinitions.h"
-#include "DatasmithC4DImportException.h"
 #include "DatasmithC4DImportOptions.h"
 #include "DatasmithC4DTranslatorModule.h"
 #include "DatasmithC4DUtils.h"
 #include "DatasmithDefinitions.h"
 #include "DatasmithImportOptions.h"
-#include "DatasmithLightImporter.h"
-#include "DatasmithMaterialExpressions.h"
 #include "DatasmithMesh.h"
 #include "DatasmithMeshExporter.h"
 #include "DatasmithMeshHelper.h"
@@ -21,14 +18,13 @@
 #include "DatasmithSceneFactory.h"
 #include "DatasmithUtils.h"
 #include "IDatasmithSceneElements.h"
-#include "Utility/DatasmithImporterUtils.h"
 
 #include "AssetRegistryModule.h"
-#include "AssetSelection.h"
 #include "Curves/RichCurve.h"
 #include "Engine/StaticMesh.h"
 #include "Engine/StaticMeshActor.h"
 #include "Engine/World.h"
+#include "Framework/Notifications/NotificationManager.h"
 #include "HAL/Event.h"
 #include "HAL/FileManager.h"
 #include "HAL/Runnable.h"
@@ -41,10 +37,9 @@
 #include "Misc/FileHelper.h"
 #include "Misc/Paths.h"
 #include "Misc/ScopedSlowTask.h"
-#include "ObjectTools.h"
-#include "PackageTools.h"
 #include "RawMesh.h"
 #include "StaticMeshAttributes.h"
+#include "Widgets/Notifications/SNotificationList.h"
 
 #include "ImathMatrixAlgo.h"
 
@@ -417,7 +412,7 @@ melange::BaseObject* FDatasmithC4DImporter::GetBestMelangeCache(melange::BaseObj
 	return ObjectCache;
 }
 
-FString FDatasmithC4DImporter::MelangeObjectID(melange::BaseObject* Object)
+TOptional<FString> FDatasmithC4DImporter::MelangeObjectID(melange::BaseObject* Object)
 {
 	//Make sure that Object is not in a cache
 	FString HierarchyPosition;
@@ -448,112 +443,123 @@ FString FDatasmithC4DImporter::MelangeObjectID(melange::BaseObject* Object)
 		}
 	}
 
-	FString ObjectID = GetMelangeBaseList2dID(Object);
-	if (InCache)
+	TOptional<FString> MelangeID = GetMelangeBaseList2dID(Object);
+	if (MelangeID && InCache)
 	{
-		ObjectID += HierarchyPosition.Right(HierarchyPosition.Len() - HierarchyPosition.Find("_C", ESearchCase::CaseSensitive) - 2);
+		MelangeID.GetValue() += HierarchyPosition.Right(HierarchyPosition.Len() - HierarchyPosition.Find("_C", ESearchCase::CaseSensitive) - 2);
 	}
-	return ObjectID;
+
+	return MelangeID;
 }
 
-// Returns whether we can remove this actor when optimizing the actor hierarchy
-bool CanRemoveActor(const TSharedPtr<IDatasmithActorElement>& Actor, const TSet<FString>& ActorNamesToKeep)
+namespace FC4DImporterImpl
 {
-	if (Actor->IsA(EDatasmithElementType::Camera | EDatasmithElementType::Light))
+	// Returns whether we can remove this actor when optimizing the actor hierarchy
+	bool CanRemoveActor(const TSharedPtr<IDatasmithActorElement>& Actor, const TSet<FString>& ActorNamesToKeep, TSharedRef<IDatasmithScene>& DatasmithScene)
 	{
-		return false;
-	}
+		if (Actor->IsA(EDatasmithElementType::Camera | EDatasmithElementType::Light))
+		{
+			return false;
+		}
 
-	if (Actor->IsA(EDatasmithElementType::StaticMeshActor))
-	{
-		TSharedPtr<IDatasmithMeshActorElement> MeshActor = StaticCastSharedPtr<IDatasmithMeshActorElement>(Actor);
-		if (MeshActor->GetStaticMeshPathName() != FString())
+		if (Actor->IsA(EDatasmithElementType::StaticMeshActor))
+		{
+			TSharedPtr<IDatasmithMeshActorElement> MeshActor = StaticCastSharedPtr<IDatasmithMeshActorElement>(Actor);
+			if (MeshActor->GetStaticMeshPathName() != FString())
+			{
+				return false;
+			}
+		}
+
+		if (DatasmithScene->GetMetaData(Actor))
 		{
 			return false;
 		}
-	}
 
-	if (ActorNamesToKeep.Contains(FString(Actor->GetName())))
-	{
-		return false;
-	}
+		if (ActorNamesToKeep.Contains(FString(Actor->GetName())))
+		{
+			return false;
+		}
 
-	return true;
-}
+		return true;
+	}
 
-void RemoveEmptyActorsRecursive(TSharedPtr<IDatasmithActorElement>& Actor, const TSet<FString>& NamesOfActorsToKeep)
-{
-	// We can't access the parent of a IDatasmithActorElement, so we have to analyze children and remove grandchildren
-	// This is also why we need a RootActor in the scene, or else we won't be able to analyze top-level actors
-	for (int32 ChildIndex = Actor->GetChildrenCount() - 1; ChildIndex >= 0; --ChildIndex)
+	void RemoveEmptyActorsRecursive(TSharedPtr<IDatasmithActorElement>& Actor, const TSet<FString>& NamesOfActorsToKeep, TSharedRef<IDatasmithScene>& DatasmithScene)
 	{
-		// Have to recurse first or else we will also iterate on our grandchildren
-		TSharedPtr<IDatasmithActorElement> Child = Actor->GetChild(ChildIndex);
+		// We can't access the parent of a IDatasmithActorElement, so we have to analyze children and remove grandchildren
+		// This is also why we need a RootActor in the scene, or else we won't be able to analyze top-level actors
+		for (int32 ChildIndex = Actor->GetChildrenCount() - 1; ChildIndex >= 0; --ChildIndex)
+		{
+			// Have to recurse first or else we will also iterate on our grandchildren
+			TSharedPtr<IDatasmithActorElement> Child = Actor->GetChild(ChildIndex);
 
-		RemoveEmptyActorsRecursive(Child, NamesOfActorsToKeep);
+			RemoveEmptyActorsRecursive(Child, NamesOfActorsToKeep, DatasmithScene);
 
-		// Move grandchildren to children
-		if (Child->GetChildrenCount() <= 1 && CanRemoveActor(Child, NamesOfActorsToKeep))
-		{
-			for (int32 GrandChildIndex = Child->GetChildrenCount() - 1; GrandChildIndex >= 0; --GrandChildIndex)
+			// Move grandchildren to children
+			if (Child->GetChildrenCount() <= 1 && CanRemoveActor(Child, NamesOfActorsToKeep, DatasmithScene))
 			{
-				TSharedPtr<IDatasmithActorElement> GrandChild = Child->GetChild(GrandChildIndex);
+				for (int32 GrandChildIndex = Child->GetChildrenCount() - 1; GrandChildIndex >= 0; --GrandChildIndex)
+				{
+					TSharedPtr<IDatasmithActorElement> GrandChild = Child->GetChild(GrandChildIndex);
 
-				Child->RemoveChild(GrandChild);
-				Actor->AddChild(GrandChild);
-			}
+					Child->RemoveChild(GrandChild);
+					Actor->AddChild(GrandChild);
+				}
 
-			Actor->RemoveChild(Child);
+				Actor->RemoveChild(Child);
+			}
 		}
 	}
-}
-
-void FDatasmithC4DImporter::RemoveEmptyActors()
-{
-	TSet<FString> NamesOfActorsToKeep;
-	NamesOfActorsToKeep.Append(NamesOfCameraTargetActors);
-	NamesOfActorsToKeep.Append(NamesOfAnimatedActors);
 
-	for (int32 ActorIndex = 0; ActorIndex < DatasmithScene->GetActorsCount(); ++ActorIndex)
+	// For now, we can't remove parents of animated nodes because animations are stored wrt local coordinate
+	// system. If we optimized an otherwise useless intermediate node, we'd need to bake its transform into all animations of
+	// child nodes, which I'm not sure is the ideal behavior as imported animation curves would look very different
+	bool KeepParentsOfAnimatedNodes(const TSharedPtr<IDatasmithActorElement>& Actor, TSet<FString>& NamesOfAnimatedActors)
 	{
-		TSharedPtr<IDatasmithActorElement> Actor = DatasmithScene->GetActor(ActorIndex);
-		RemoveEmptyActorsRecursive(Actor, NamesOfActorsToKeep);
-	}
-}
+		bool bKeepThisNode = NamesOfAnimatedActors.Contains(Actor->GetName());
 
-// For now, we can't remove parents of animated nodes because animations are stored wrt local coordinate
-// system. If we optimized an otherwise useless intermediate node, we'd need to bake its transform into all animations of
-// child nodes, which I'm not sure is the ideal behavior as imported animation curves would look very different
-bool KeepParentsOfAnimatedNodes(const TSharedPtr<IDatasmithActorElement>& Actor, TSet<FString>& NamesOfAnimatedActors)
-{
-	bool bKeepThisNode = NamesOfAnimatedActors.Contains(Actor->GetName());
+		for (int32 ChildIndex = 0; ChildIndex < Actor->GetChildrenCount(); ++ChildIndex)
+		{
+			bKeepThisNode |= KeepParentsOfAnimatedNodes(Actor->GetChild(ChildIndex), NamesOfAnimatedActors);
+		}
 
-	for (int32 ChildIndex = 0; ChildIndex < Actor->GetChildrenCount(); ++ChildIndex)
-	{
-		bKeepThisNode |= KeepParentsOfAnimatedNodes(Actor->GetChild(ChildIndex), NamesOfAnimatedActors);
+		if (bKeepThisNode)
+		{
+			NamesOfAnimatedActors.Add(Actor->GetName());
+		}
+
+		return bKeepThisNode;
 	}
 
-	if (bKeepThisNode)
+	void RemoveEmptyActors(TSharedRef<IDatasmithScene>& DatasmithScene, const TSet<FString>& NamesOfCameraTargetActors, const TSet<FString>& NamesOfAnimatedActors)
 	{
-		NamesOfAnimatedActors.Add(Actor->GetName());
-	}
+		TSet<FString> NamesOfActorsToKeep;
+		NamesOfActorsToKeep.Append(NamesOfCameraTargetActors);
+		NamesOfActorsToKeep.Append(NamesOfAnimatedActors);
 
-	return bKeepThisNode;
+		for (int32 ActorIndex = 0; ActorIndex < DatasmithScene->GetActorsCount(); ++ActorIndex)
+		{
+			TSharedPtr<IDatasmithActorElement> Actor = DatasmithScene->GetActor(ActorIndex);
+			RemoveEmptyActorsRecursive(Actor, NamesOfActorsToKeep, DatasmithScene);
+		}
+	}
 }
 
-TSharedPtr<IDatasmithMetaDataElement> FDatasmithC4DImporter::CreateMetadataForActor(const IDatasmithActorElement& Actor)
+namespace FC4DImporterImpl
 {
-	TSharedPtr<IDatasmithMetaDataElement>& Metadata = ActorMetadata.FindOrAdd(&Actor);
-	if (!Metadata.IsValid())
+	TSharedPtr<IDatasmithMetaDataElement> CreateMetadataForActor(const TSharedPtr<IDatasmithActorElement>& Actor, const TSharedRef<IDatasmithScene>& DatasmithScene)
 	{
-		Metadata = FDatasmithSceneFactory::CreateMetaData(Actor.GetName());
+		if (!Actor.IsValid())
+		{
+			return nullptr;
+		}
+
+		TSharedPtr<IDatasmithMetaDataElement> Metadata = FDatasmithSceneFactory::CreateMetaData(Actor->GetName());
+		Metadata->SetAssociatedElement(Actor);
+		DatasmithScene->AddMetaData(Metadata);
+		return Metadata;
 	}
-	DatasmithScene->AddMetaData(Metadata);
-	return Metadata;
-}
 
-namespace
-{
 	void AddMetadataVector(IDatasmithMetaDataElement* Metadata, const FString& Key, const FVector& Value)
 	{
 		TSharedPtr<IDatasmithKeyValueProperty> MetadataPropertyPtr = FDatasmithSceneFactory::CreateKeyValueProperty(*Key);
@@ -607,64 +613,88 @@ namespace
 
 		Metadata->AddProperty(MetadataPropertyPtr);
 	}
-}
 
-void FDatasmithC4DImporter::AddChildActor(melange::BaseObject* Object, TSharedPtr<IDatasmithActorElement> ParentActor, melange::Matrix WorldTransformMatrix, const TSharedPtr<IDatasmithActorElement>& Actor)
-{
-	melange::DynamicDescription* DynamicDescription = Object->GetDynamicDescription();
-	if (DynamicDescription)
+	void ImportActorMetadata(melange::BaseObject* Object, const TSharedPtr<IDatasmithActorElement>& Actor, TSharedRef<IDatasmithScene>& DatasmithScene)
 	{
-		TSharedPtr<IDatasmithMetaDataElement> Metadata = CreateMetadataForActor(*Actor);
+		melange::DynamicDescription* DynamicDescription = Object->GetDynamicDescription();
+		if (!DynamicDescription)
+		{
+			return;
+		}
+
+		TSharedPtr<IDatasmithMetaDataElement> Metadata = nullptr;
 
 		void* BrowserHandle = DynamicDescription->BrowseInit();
 		melange::DescID DescId;
 		const melange::BaseContainer* DescContainer;
 		while (DynamicDescription->BrowseGetNext(BrowserHandle, &DescId, &DescContainer))
 		{
-			melange::Int32 UserDataType = DescContainer->GetInt32(21);
-			melange::GeData Data;
-			if (Object->GetParameter(DescId, Data) && Data.GetType() != melange::DA_NIL)
+			if (DescId[0].id != melange::ID_USERDATA)
 			{
-				FString BaseDataName = MelangeStringToFString(DescContainer->GetString(1));
-				FString DataName = BaseDataName;
-				FString CurrentValue;
-				int DataNameIndex = 0;
-				if (UserDataType == melange::DA_VECTOR)
-				{
-					FVector ConvertedVector = ConvertMelangePosition(Data.GetVector());
-					AddMetadataVector(Metadata.Get(), *DataName, ConvertedVector);
-				}
-				else if (UserDataType == melange::DA_REAL)
-				{
-					AddMetadataFloat(Metadata.Get(), *DataName, static_cast<float>(Data.GetFloat()));
-				}
-				else if (UserDataType == 1000492 /*color*/)
-				{
-					AddMetadataColor(Metadata.Get(), *DataName, MelangeVectorToFVector(Data.GetVector()));
-				}
-				else if (UserDataType == 1000484 /*texture*/)
-				{
-					AddMetadataTexture(Metadata.Get(), *DataName, *GeDataToString(Data));
-				}
-				else if (UserDataType == 400006001 /*boolean*/)
+				continue;;
+			}
+
+			if (!Metadata.IsValid())
+			{
+				Metadata = CreateMetadataForActor(Actor, DatasmithScene);
+				if (!Metadata.IsValid())
 				{
-					AddMetadataBool(Metadata.Get(), *DataName, Data.GetInt32() != 0);
+					return;
 				}
-				else
+			}
+
+			melange::GeData Data;
+			if (!Object->GetParameter(DescId, Data))
+			{
+				continue;
+			}
+
+			FString DataName = MelangeStringToFString(DescContainer->GetString(melange::DESC_NAME));
+
+			melange::Int32 UserDataType = DescContainer->GetInt32(melange::DESC_CUSTOMGUI);
+			if (UserDataType == melange::DA_VECTOR)
+			{
+				FVector ConvertedVector = ConvertMelangePosition(Data.GetVector());
+				AddMetadataVector(Metadata.Get(), *DataName, ConvertedVector);
+			}
+			else if (UserDataType == melange::DA_REAL)
+			{
+				AddMetadataFloat(Metadata.Get(), *DataName, static_cast<float>(Data.GetFloat()));
+			}
+			else if (UserDataType == 1000492 /*color*/)
+			{
+				AddMetadataColor(Metadata.Get(), *DataName, MelangeVectorToFVector(Data.GetVector()));
+			}
+			else if (UserDataType == 1000484 /*texture*/)
+			{
+				AddMetadataTexture(Metadata.Get(), *DataName, *GeDataToString(Data));
+			}
+			else if (UserDataType == 400006001 /*boolean*/)
+			{
+				AddMetadataBool(Metadata.Get(), *DataName, Data.GetInt32() != 0);
+			}
+			else
+			{
+				FString ValueString = GeDataToString(Data);
+				if (!ValueString.IsEmpty())
 				{
-					FString ValueString = GeDataToString(Data);
-					if (!ValueString.IsEmpty())
-					{
-						AddMetadataString(Metadata.Get(), *DataName, *ValueString);
-					}
+					AddMetadataString(Metadata.Get(), *DataName, *ValueString);
 				}
-
 			}
 		}
 		DynamicDescription->BrowseFree(BrowserHandle);
 	}
+}
 
-	DatasmithC4DImportCheck(!NamesOfAllActors.Contains(Actor->GetName()));
+bool FDatasmithC4DImporter::AddChildActor(melange::BaseObject* Object, TSharedPtr<IDatasmithActorElement> ParentActor, melange::Matrix WorldTransformMatrix, const TSharedPtr<IDatasmithActorElement>& Actor)
+{
+	FC4DImporterImpl::ImportActorMetadata(Object, Actor, DatasmithScene);
+
+	if (NamesOfAllActors.Contains(Actor->GetName()))
+	{
+		//Duplicate name, don't import twice.
+		return false;
+	}
 	NamesOfAllActors.Add(Actor->GetName());
 
 	ActorElementToC4DObject.Add(Actor.Get(), Object);
@@ -713,9 +743,8 @@ void FDatasmithC4DImporter::AddChildActor(melange::BaseObject* Object, TSharedPt
 	bool bExtracted = Imath::extractAndRemoveScalingAndShear<float>(Matrix, Scale, Shear, false);
 	if (!bExtracted)
 	{
-		// TODO: Append a message to the build summary.
-		FString Msg = FString::Printf(TEXT("WARNING: Actor %ls (%ls) has some zero scaling"), Actor->GetName(), Actor->GetLabel());
-		return;
+		UE_LOG(LogDatasmithC4DImport, Warning, TEXT("Actor %ls (%ls) has some zero scaling"), Actor->GetName(), Actor->GetLabel());
+		return false;
 	}
 
 	// Initialize a rotation quaternion with the rotation matrix.
@@ -763,6 +792,8 @@ void FDatasmithC4DImporter::AddChildActor(melange::BaseObject* Object, TSharedPt
 	{
 		Actor->AddChild(Child, EDatasmithActorAttachmentRule::KeepWorldTransform);
 	}
+
+	return true;
 }
 
 TSharedPtr<IDatasmithActorElement> FDatasmithC4DImporter::ImportNullActor(melange::BaseObject* Object, const FString& DatasmithName, const FString& DatasmithLabel)
@@ -1087,10 +1118,14 @@ TSharedPtr<IDatasmithCameraActorElement> FDatasmithC4DImporter::ImportCamera(mel
 	if (LookAtObject)
 	{
 		//LookAtObject can not be a cached object or an instanced object so GetMelangeBaseList2dID should be the final ID
-		FString LookAtID = GetMelangeBaseList2dID(LookAtObject);
-		CameraActor->SetLookAtActor(*LookAtID);
+		TOptional<FString> LookAtID = GetMelangeBaseList2dID(LookAtObject);
+		if (!LookAtID)
+		{
+			return TSharedPtr<IDatasmithCameraActorElement>();
+		}
+		CameraActor->SetLookAtActor(*(LookAtID.GetValue()));
 		CameraActor->SetLookAtAllowRoll(true);
-		NamesOfCameraTargetActors.Add(LookAtID);
+		NamesOfCameraTargetActors.Add(LookAtID.GetValue());
 	}
 
 	float CameraFocusDistanceInCM = MelangeGetFloat(InC4DCameraPtr, melange::CAMERAOBJECT_TARGETDISTANCE);
@@ -1188,10 +1223,14 @@ FString FDatasmithC4DImporter::GetBaseShaderTextureFilePath(melange::BaseList2D*
 
 TSharedPtr<IDatasmithMasterMaterialElement> FDatasmithC4DImporter::ImportMaterial(melange::Material* InC4DMaterialPtr)
 {
-	FString DatasmithName  = GetMelangeBaseList2dID(InC4DMaterialPtr);
+	TOptional<FString> DatasmithName = GetMelangeBaseList2dID(InC4DMaterialPtr);
+	if (!DatasmithName)
+	{
+		return TSharedPtr<IDatasmithMasterMaterialElement>();
+	}
 	FString DatasmithLabel = FDatasmithUtils::SanitizeObjectName(MelangeObjectName(InC4DMaterialPtr));
 
-	TSharedPtr<IDatasmithMasterMaterialElement> MaterialPtr = FDatasmithSceneFactory::CreateMasterMaterial(*DatasmithName);
+	TSharedPtr<IDatasmithMasterMaterialElement> MaterialPtr = FDatasmithSceneFactory::CreateMasterMaterial(*(DatasmithName.GetValue()));
 	MaterialPtr->SetLabel(*DatasmithLabel);
 	MaterialPtr->SetMaterialType(EDatasmithMasterMaterialType::Opaque);
 
@@ -1489,7 +1528,7 @@ TSharedPtr<IDatasmithMasterMaterialElement> FDatasmithC4DImporter::ImportMateria
 	return MaterialPtr;
 }
 
-void FDatasmithC4DImporter::ImportMaterialHierarchy(melange::BaseMaterial* InC4DMaterialPtr)
+bool FDatasmithC4DImporter::ImportMaterialHierarchy(melange::BaseMaterial* InC4DMaterialPtr)
 {
 	// Reinitialize the scene material map and texture set.
 	MaterialNameToMaterialElement.Empty();
@@ -1498,10 +1537,18 @@ void FDatasmithC4DImporter::ImportMaterialHierarchy(melange::BaseMaterial* InC4D
 	{
 		if (InC4DMaterialPtr->GetType() == Mmaterial)
 		{
-			TSharedPtr<IDatasmithMasterMaterialElement> DatasmithMaterial = ImportMaterial(static_cast<melange::Material*>(InC4DMaterialPtr));
-			MaterialNameToMaterialElement.Add(DatasmithMaterial->GetName(), DatasmithMaterial);
+			if (TSharedPtr<IDatasmithMasterMaterialElement> DatasmithMaterial = ImportMaterial(static_cast<melange::Material*>(InC4DMaterialPtr)))
+			{
+				MaterialNameToMaterialElement.Add(DatasmithMaterial->GetName(), DatasmithMaterial);
+			}
+			else
+			{
+				return false;
+			}
 		}
 	}
+
+	return true;
 }
 
 FString FDatasmithC4DImporter::CustomizeMaterial(const FString& InMaterialID, const FString& InMeshID, melange::TextureTag* InTextureTag)
@@ -1548,51 +1595,83 @@ FString FDatasmithC4DImporter::CustomizeMaterial(const FString& InMaterialID, co
 	return InMaterialID;
 }
 
-TSharedPtr<IDatasmithMeshActorElement> FDatasmithC4DImporter::ImportPolygon(melange::PolygonObject* PolyObject, TMap<FString, melange::PolygonObject*>* ClonerBaseChildrenHash, const FString& DatasmithName, const FString& DatasmithLabel, const TArray<melange::TextureTag*>& TextureTags)
+TMap<int32, FString> FDatasmithC4DImporter::GetCustomizedMaterialAssignment(const FString& DatasmithMeshName, const TArray<melange::TextureTag*>& TextureTags)
 {
-	FMD5Hash PolygonHash;
+	TMap<int32, FString> SlotToMaterialName;
 
-	melange::PolygonObject* DataPolyObject = PolyObject;
-	if (ClonerBaseChildrenHash)
+	// Create customized materials for all the used texture tags. This because each tag
+	// actually represents a material "instance", and might have different settings like texture tiling
+	for (int32 SlotIndex = 0; SlotIndex < TextureTags.Num(); ++SlotIndex)
 	{
-		PolygonHash = ComputePolygonDataHash(PolyObject);
-		FString PolygonHashStr = BytesToHex(PolygonHash.GetBytes(), 16);
+		melange::TextureTag* Tag = TextureTags[SlotIndex];
 
-		if (melange::PolygonObject** BasePolygon = ClonerBaseChildrenHash->Find(PolygonHashStr))
-		{
-			DataPolyObject = *BasePolygon;
-		}
-		else
+		FString CustomizedMaterialName;
+		melange::BaseList2D* TextureMaterial = Tag ? MelangeGetLink(Tag, melange::TEXTURETAG_MATERIAL) : nullptr;
+		if (TextureMaterial)
 		{
-			ClonerBaseChildrenHash->Add(PolygonHashStr, PolyObject);
+			// This can also return an existing material without necessarily spawning a new instance
+			if (TOptional<FString> MaterialID = GetMelangeBaseList2dID(TextureMaterial))
+			{
+				CustomizedMaterialName = CustomizeMaterial(MaterialID.GetValue(), DatasmithMeshName, Tag);
+			}
 		}
+
+		SlotToMaterialName.Add(SlotIndex, CustomizedMaterialName);
 	}
 
+	return SlotToMaterialName;
+}
+
+TSharedPtr<IDatasmithMeshActorElement> FDatasmithC4DImporter::ImportPolygon(melange::PolygonObject* PolyObject, const FString& DatasmithActorName, const FString& DatasmithActorLabel, const TArray<melange::TextureTag*>& TextureTags)
+{
+	FMD5Hash PolygonHash = ComputePolygonDataHash(PolyObject);
+	FString HashString = BytesToHex(PolygonHash.GetBytes(), PolygonHash.GetSize());
+
+	TOptional<FString> MelangeID = MelangeObjectID(PolyObject);
+	if (!MelangeID.IsSet())
+	{
+		return nullptr;
+	}
+	const FString& DatasmithMeshName = MelangeID.GetValue();
+
 	IDatasmithMeshElement* ResultMeshElement = nullptr;
-	if (TSharedRef<IDatasmithMeshElement>* PreviousMesh = PolygonObjectToMeshElement.Find(DataPolyObject))
+	if (TSharedRef<IDatasmithMeshElement>* PreviousMesh = PolygonHashToMeshElement.Find(HashString))
 	{
 		ResultMeshElement = &PreviousMesh->Get();
 	}
 	else
 	{
-		// Compute the hash if we haven't by chance above
-		if (!PolygonHash.IsValid())
-		{
-			PolygonHash = ComputePolygonDataHash(PolyObject);
-		}
-
-		TSharedPtr<IDatasmithMeshElement> MeshElement = ImportMesh(PolyObject, MelangeObjectID(DataPolyObject), DatasmithLabel, TextureTags);
+		TSharedPtr<IDatasmithMeshElement> MeshElement = ImportMesh(PolyObject, *DatasmithMeshName, DatasmithActorLabel);
 		ResultMeshElement = MeshElement.Get();
 
 		// Set the polygon hash as the file hash. It will be checked by Datasmith in
 		// FDatasmithImporter::FilterElementsToImport to know if a mesh has changed and
 		// the asset needs to be replaced during reimport
 		ResultMeshElement->SetFileHash(PolygonHash);
+
+		PolygonHashToMeshElement.Add(HashString, MeshElement.ToSharedRef());
 	}
 
-	TSharedPtr<IDatasmithMeshActorElement> MeshActorElement = FDatasmithSceneFactory::CreateMeshActor(*DatasmithName);
-	MeshActorElement->SetLabel(*DatasmithLabel);
+	TSharedPtr<IDatasmithMeshActorElement> MeshActorElement = FDatasmithSceneFactory::CreateMeshActor(*DatasmithActorName);
+	MeshActorElement->SetLabel(*DatasmithActorLabel);
 	MeshActorElement->SetStaticMeshPathName(ResultMeshElement->GetName());
+
+	// Add material overrides
+	TMap<int32, FString> SlotIndexToMaterialName = GetCustomizedMaterialAssignment(DatasmithMeshName, TextureTags);
+	for (const TPair<int32, FString>& Pair : SlotIndexToMaterialName)
+	{
+		int32 SlotIndex = Pair.Key;
+		const FString& MaterialName = Pair.Value;
+
+		TSharedPtr<IDatasmithMasterMaterialElement>* FoundMaterial = MaterialNameToMaterialElement.Find(MaterialName);
+		if (FoundMaterial && FoundMaterial->IsValid())
+		{
+			TSharedRef<IDatasmithMaterialIDElement> MaterialIDElement = FDatasmithSceneFactory::CreateMaterialId(*MaterialName);
+			MaterialIDElement->SetId(SlotIndex);
+			MeshActorElement->AddMaterialOverride(MaterialIDElement);
+		}
+	}
+
 	return MeshActorElement;
 }
 
@@ -2190,7 +2269,89 @@ void FDatasmithC4DImporter::ImportActorHierarchyAnimations(TSharedPtr<IDatasmith
 	}
 }
 
-TSharedPtr<IDatasmithActorElement> FDatasmithC4DImporter::ImportObjectAndChildren(melange::BaseObject* ActorObject, melange::BaseObject* DataObject, TSharedPtr<IDatasmithActorElement> ParentActor, const melange::Matrix& WorldTransformMatrix, TMap<FString, melange::PolygonObject*>* ClonerBaseChildrenHash, const FString& InstancePath, TArray<melange::BaseObject*>* InstanceObjects, TArray<melange::TextureTag*> TextureTags, const FString& DatasmithLabel)
+TArray<melange::TextureTag*> FDatasmithC4DImporter::GetActiveTextureTags(const melange::BaseObject* Object)
+{
+	if (!Object)
+	{
+		return {};
+	}
+
+	// If we're the cache of an object, then our texture tags will be stored on that object.
+	// This repeats recursively, with only the first 'non-cache' BaseObject actually containing the TextureTags.
+	// Most PolygonObjects are not caches, but if a polygon has direct children then this will happen.
+	const melange::BaseObject* ObjectToSearch = Object;
+	while (const melange::BaseObject* Parent = CachesOriginalObject.FindRef(ObjectToSearch))
+	{
+		ObjectToSearch = Parent;
+	}
+
+	if (ObjectToSearch == nullptr)
+	{
+		return {};
+	}
+
+	TArray<melange::TextureTag*> OrderedTextureTags;
+	TArray<melange::BaseSelect*> OrderedSelectionTags;
+	OrderedSelectionTags.Add(nullptr); // "unselected" group
+
+	TMap<FString, melange::BaseSelect*> SelectionTagsByName;
+
+	for (melange::BaseTag* Tag = const_cast<melange::BaseObject*>(ObjectToSearch)->GetFirstTag(); Tag; Tag = Tag->GetNext())
+	{
+		melange::Int32 TagType = Tag->GetType();
+		if (TagType == Ttexture)
+		{
+			OrderedTextureTags.Add(static_cast<melange::TextureTag*>(Tag));
+		}
+		else if (TagType == Tpolygonselection)
+		{
+			FString SelectionName = MelangeGetString(Tag, melange::POLYGONSELECTIONTAG_NAME);
+			if (!SelectionName.IsEmpty())
+			{
+				melange::BaseSelect* BaseSelect = static_cast<melange::SelectionTag*>(Tag)->GetBaseSelect();
+				OrderedSelectionTags.Add(BaseSelect);
+
+				SelectionTagsByName.Add(SelectionName, BaseSelect);
+			}
+		}
+	}
+
+	// If we have multiple texture tags using the same selection, only the latter one will be applied
+	// Order is important here: We must scan TextureTags front to back to guarantee that behavior
+	TMap<melange::BaseSelect*, melange::TextureTag*> ActiveTextureTags;
+	for (melange::TextureTag* TextureTag : OrderedTextureTags)
+	{
+		melange::BaseSelect* UsedSelectionTag = nullptr;
+
+		FString UsedSelectionTagName = MelangeGetString(TextureTag, melange::TEXTURETAG_RESTRICTION);
+		if (!UsedSelectionTagName.IsEmpty())
+		{
+			if (melange::BaseSelect** FoundSelectionTag = SelectionTagsByName.Find(UsedSelectionTagName))
+			{
+				UsedSelectionTag = *FoundSelectionTag;
+			}
+		}
+
+		melange::BaseList2D* TextureMaterial = TextureTag ? MelangeGetLink(TextureTag, melange::TEXTURETAG_MATERIAL) : nullptr;
+
+		// Note: If this texture tag is applied without a polygon selection, UsedSelectionTag will be
+		// nullptr here, but that is intentional: It's how we signal the "unselected" selection group
+		ActiveTextureTags.Add(UsedSelectionTag, TextureTag);
+	}
+
+	// Order is important: The polygon groups are created according to the order with which we retrieve our selection tags,
+	// so the order with which we return these texture tags must match it exactly
+	TArray<melange::TextureTag*> OrderedActiveTextureTags;
+	for (melange::BaseSelect* SelectionTag : OrderedSelectionTags)
+	{
+		melange::TextureTag** TextureTagForSelection = ActiveTextureTags.Find(SelectionTag);
+		OrderedActiveTextureTags.Add(TextureTagForSelection ? *TextureTagForSelection : nullptr);
+	}
+
+	return OrderedActiveTextureTags;
+}
+
+TSharedPtr<IDatasmithActorElement> FDatasmithC4DImporter::ImportObjectAndChildren(melange::BaseObject* ActorObject, melange::BaseObject* DataObject, TSharedPtr<IDatasmithActorElement> ParentActor, const melange::Matrix& WorldTransformMatrix, const FString& InstancePath, TArray<melange::BaseObject*>* InstanceObjects, const FString& DatasmithLabel)
 {
 	TSharedPtr<IDatasmithActorElement> ActorElement;
 	melange::Int32 ObjectType = DataObject->GetType();
@@ -2204,22 +2365,17 @@ TSharedPtr<IDatasmithActorElement> FDatasmithC4DImporter::ImportObjectAndChildre
 	{
 		ActorCache = DataCache;
 	}
-	FString DatasmithName = MelangeObjectID(ActorObject);
-	if (!InstancePath.IsEmpty())
+
+	TOptional<FString> DatasmithName = MelangeObjectID(ActorObject);
+	if (!DatasmithName)
 	{
-		DatasmithName = MD5FromString(InstancePath) + "_" + DatasmithName;
+		UE_LOG(LogDatasmithC4DImport, Error, TEXT("Could not get the ID of object \"%s\": %s"), *MelangeObjectName(ActorObject));
+		DatasmithName = FString(TEXT("Invalid object"));
 	}
 
-	//Get all texture tags
-	melange::BaseTag* Tag = ActorObject->GetFirstTag();
-	while (Tag)
+	if (!InstancePath.IsEmpty())
 	{
-		melange::Int32 TagType = Tag->GetType();
-		if (TagType == Ttexture)
-		{
-			TextureTags.Add(static_cast<melange::TextureTag*>(Tag));
-		}
-		Tag = Tag->GetNext();
+		DatasmithName = MD5FromString(InstancePath) + "_" + DatasmithName.GetValue();
 	}
 
 	melange::Matrix NewWorldTransformMatrix = WorldTransformMatrix * ActorObject->GetMl();
@@ -2243,46 +2399,58 @@ TSharedPtr<IDatasmithActorElement> FDatasmithC4DImporter::ImportObjectAndChildre
 
 	if (bActorVisible)
 	{
-		try
+		bool bSuccess = true;
+
+		if (InstanceObjects)
 		{
-			if (InstanceObjects)
+			if (InstanceObjects->Num() > 0 && (*InstanceObjects)[0]->GetType() == ObjectType)
 			{
-				DatasmithC4DImportCheck(InstanceObjects->Num() > 0);
-				melange::BaseObject* RealDataObject = (*InstanceObjects)[0];
-				InstanceObjects->RemoveAt(0);
-				DatasmithC4DImportCheck(RealDataObject->GetType() == ObjectType);
-				DataObject = RealDataObject;
+				DataObject = (*InstanceObjects)[0];
 			}
+			else
+			{
+				bSuccess = false;
+			}
+		}
 
+		if (bSuccess)
+		{
 			if (ObjectType == Oinstance)
 			{
-				melange::BaseObject* InstanceLink = static_cast<melange::BaseObject*>(MelangeGetLink(DataObject, melange::INSTANCEOBJECT_LINK));
-				DatasmithC4DImportCheck(InstanceLink != nullptr);
-				TArray<melange::BaseObject*> CurentInstanceObjects = GetMelangeInstanceObjects(InstanceLink);
-				melange::Int32 LinkObjectType = InstanceLink->GetType();
-				if (ActorCache)
+				if (melange::BaseObject* InstanceLink = static_cast<melange::BaseObject*>(MelangeGetLink(DataObject, melange::INSTANCEOBJECT_LINK)))
 				{
-					// Parse our own duplicated hierarchy (which is a replica of the original object's hierarchy),
-					// carrying our own texture tags. If we jump through InstanceLink, we'll be parsing the original
-					// hierarchy, so any animations or polygons we parse will be bound to the original actors (not
-					// our replica actors)
-					return ImportObjectAndChildren(ActorObject, ActorCache, ParentActor, WorldTransformMatrix, ClonerBaseChildrenHash, MelangeObjectID(DataObject) + InstancePath, &CurentInstanceObjects, TextureTags, DatasmithLabel);
-				}
-				else
-				{
-					return ImportObjectAndChildren(ActorObject, InstanceLink, ParentActor, WorldTransformMatrix, ClonerBaseChildrenHash, MelangeObjectID(DataObject) + InstancePath, &CurentInstanceObjects, TextureTags, DatasmithLabel);
+					TArray<melange::BaseObject*> CurentInstanceObjects = GetMelangeInstanceObjects(InstanceLink);
+					melange::BaseObject* ActorToImport = InstanceLink;
+
+					if (ActorCache)
+					{
+						// Parse our own duplicated hierarchy (which is a replica of the original object's hierarchy),
+						// carrying our own texture tags. If we jump through InstanceLink, we'll be parsing the original
+						// hierarchy, so any animations or polygons we parse will be bound to the original actors (not
+						// our replica actors)
+						ActorToImport = ActorCache;
+					}
+
+					if (TOptional<FString> ObjectID = MelangeObjectID(DataObject))
+					{
+						return ImportObjectAndChildren(ActorObject, ActorToImport, ParentActor, WorldTransformMatrix, ObjectID.GetValue() + InstancePath, &CurentInstanceObjects, DatasmithLabel);
+					}
 				}
-			}
 
-			// For particle emitters, we need to mark all the child actors, as those need to have their visibility
-			// manually animated to simulate mesh particles spawning and despawning
-			if (ObjectType == Oparticle)
+				bSuccess = false;
+			}
+			else if(ObjectType == Oparticle)
 			{
+				// For particle emitters, we need to mark all the child actors, as those need to have their visibility
+				// manually animated to simulate mesh particles spawning and despawning
 				MarkActorsAsParticles(ActorObject, ActorCache);
 			}
 
-			if (ObjectType == Ocloner || ObjectType == Oarray)
+			bool bImportCache = false;
+			switch (ObjectType)
 			{
+			case Ocloner:
+			case Oarray:
 				//Cloner(Ocloner)
 				//	| -CACHE: Null(Onull)
 				//	| | -Cube 2(Ocube)
@@ -2293,7 +2461,6 @@ TSharedPtr<IDatasmithActorElement> FDatasmithC4DImporter::ImportObjectAndChildre
 				//	| | | -CACHE: Cube 0(Opolygon)
 				//	| -Cube(Ocube)
 
-				melange::GeData Data;
 				if (ObjectType == Ocloner && MelangeGetInt32(ActorObject, melange::MGCLONER_VOLUMEINSTANCES_MODE) != 0)
 				{
 					// Render/Multi-instance cloner should be ignored
@@ -2306,87 +2473,115 @@ TSharedPtr<IDatasmithActorElement> FDatasmithC4DImporter::ImportObjectAndChildre
 				}
 				else
 				{
-					DatasmithC4DImportCheck(DataCache != nullptr);
-					DatasmithC4DImportCheck(DataCache->GetType() == Onull);
-					ActorElement = ImportNullActor(ActorObject, DatasmithName, DatasmithLabel);
-					AddChildActor(ActorObject, ParentActor, NewWorldTransformMatrix, ActorElement);
-					TMap<FString, melange::PolygonObject*> ThisClonerBaseChildrenHash;
-					ImportHierarchy(ActorCache->GetDown(), DataCache->GetDown(), ActorElement, NewWorldTransformMatrix, &ThisClonerBaseChildrenHash, InstancePath, nullptr, TextureTags);
+					ActorElement = ImportNullActor(ActorObject, DatasmithName.GetValue(), DatasmithLabel);
+					if (DataCache != nullptr && DataCache->GetType() == Onull && AddChildActor(ActorObject, ParentActor, NewWorldTransformMatrix, ActorElement))
+					{
+						ImportHierarchy(ActorCache->GetDown(), DataCache->GetDown(), ActorElement, NewWorldTransformMatrix, InstancePath, nullptr);
+						return ActorElement;
+					}
+
+					bSuccess = false;
+				}
+
+				break;
+
+			case Ofracture:
+			case ID_MOTIONFRACTUREVORONOI:
+			case Osymmetry:
+			case Osds: //Sub Division Surface
+			case Oboole:
+				ActorElement = ImportNullActor(ActorObject, DatasmithName.GetValue() +"0"/*to be different than the cache root*/, DatasmithLabel);
+				if (DataCache != nullptr && AddChildActor(ActorObject, ParentActor, NewWorldTransformMatrix, ActorElement)
+					&& ImportObjectAndChildren(ActorCache, DataCache, ActorElement, NewWorldTransformMatrix, InstancePath, nullptr, DatasmithLabel))
+				{
 					return ActorElement;
 				}
 
-			}
-			else if (ObjectType == Ofracture || ObjectType == ID_MOTIONFRACTUREVORONOI || ObjectType == Osymmetry || ObjectType == Osds /*Sub Division Surface*/ || ObjectType == Oboole)
-			{
-				DatasmithC4DImportCheck(DataCache != nullptr);
-				ActorElement = ImportNullActor(ActorObject, DatasmithName + "0"/*to be different than the cache root*/, DatasmithLabel);
-				AddChildActor(ActorObject, ParentActor, NewWorldTransformMatrix, ActorElement);
-				ImportObjectAndChildren(ActorCache, DataCache, ActorElement, NewWorldTransformMatrix, ClonerBaseChildrenHash, InstancePath, nullptr, TextureTags, DatasmithLabel);
-				return ActorElement;
-			}
-			else if (ObjectType == Ospline)
-			{
+				bSuccess = false;
+				break;
+
+			case Ospline:
 				if (melange::SplineObject* Spline = static_cast<melange::SplineObject*>(ActorObject))
 				{
 					ImportSpline(Spline);
 				}
+				break;
+
+			default:
+				bImportCache = true;
+				break;
 			}
-			else if (ActorCache)
+
+			if (bSuccess && bImportCache && ActorCache)
 			{
-				ActorElement = ImportObjectAndChildren(ActorCache, DataCache, nullptr, ActorCache->GetMg(), ClonerBaseChildrenHash, InstancePath, nullptr, TextureTags, DatasmithLabel);
+				ActorElement = ImportObjectAndChildren(ActorCache, DataCache, nullptr, ActorCache->GetMg(), InstancePath, nullptr, DatasmithLabel);
 			}
 			else if (ObjectType == Opolygon)
 			{
 				melange::PolygonObject* PolygonObject = static_cast<melange::PolygonObject*>(DataObject);
 				if (Options->bImportEmptyMesh || PolygonObject->GetPolygonCount() > 0)
 				{
-					ActorElement = ImportPolygon(static_cast<melange::PolygonObject*>(DataObject), ClonerBaseChildrenHash, DatasmithName, DatasmithLabel, TextureTags);
+					if (TSharedPtr<IDatasmithMeshActorElement> MeshActorElement = ImportPolygon(PolygonObject, DatasmithName.GetValue(), DatasmithLabel, GetActiveTextureTags(PolygonObject)))
+					{
+						ActorElement = MeshActorElement;
+					}
+					else
+					{
+						bSuccess = false;
+					}
 				}
 			}
 			else if (ObjectType == Ocamera)
 			{
-				ActorElement = ImportCamera(DataObject, DatasmithName, DatasmithLabel);
+				if (TSharedPtr<IDatasmithCameraActorElement> CameraElement = ImportCamera(DataObject, DatasmithName.GetValue(), DatasmithLabel))
+				{
+					ActorElement = CameraElement;
+				}
+				else
+				{
+					bSuccess = false;
+				}
 			}
 			else if (ObjectType == Olight)
 			{
-				ActorElement = ImportLight(DataObject, DatasmithName, DatasmithLabel);
+				ActorElement = ImportLight(DataObject, DatasmithName.GetValue(), DatasmithLabel);
 			}
 		}
-		catch (const DatasmithC4DImportException& Exception)
-		{
-			UE_LOG(LogDatasmithC4DImport, Error, TEXT("Could not import the object \"%s\": %s"), *MelangeObjectName(ActorObject), *Exception.GetMessage());
-		}
-	}
 
-	try
-	{
-		if (!ActorElement.IsValid())
+		if (!bSuccess)
 		{
-			ActorElement = ImportNullActor(ActorObject, DatasmithName, DatasmithLabel);
+			UE_LOG(LogDatasmithC4DImport, Error, TEXT("Could not import the object \"%s\""), *MelangeObjectName(ActorObject));
 		}
 
-		if (ParentActor)
-		{
-			AddChildActor(ActorObject, ParentActor, NewWorldTransformMatrix, ActorElement);
-		}
+	} //if(bActorVisible)
 
-		// Invisible layers will not be imported, so don't use their names
-		if (bActorVisible)
+	if (!ActorElement.IsValid())
+	{
+		ActorElement = ImportNullActor(ActorObject, DatasmithName.GetValue(), DatasmithLabel);
+	}
+
+	bool bSuccessfullyAddedChildActor = true;
+	if (ParentActor)
+	{
+		if (!AddChildActor(ActorObject, ParentActor, NewWorldTransformMatrix, ActorElement))
 		{
-			ActorElement->SetLayer(*TargetLayerName);
+			bSuccessfullyAddedChildActor = false;
+			UE_LOG(LogDatasmithC4DImport, Error, TEXT("Could not create the actor for the object \"%s\""), *MelangeObjectName(ActorObject));
 		}
 	}
-	catch (const DatasmithC4DImportException& Exception)
+
+	// Invisible layers will not be imported, so don't use their names
+	if (bActorVisible && bSuccessfullyAddedChildActor)
 	{
-		UE_LOG(LogDatasmithC4DImport, Error, TEXT("Could not create the actor for the object \"%s\": %s"), *MelangeObjectName(ActorObject), *Exception.GetMessage());
+		ActorElement->SetLayer(*TargetLayerName);
 	}
 
-	ImportHierarchy(ActorObject->GetDown(), DataObject->GetDown(), ActorElement, NewWorldTransformMatrix, ClonerBaseChildrenHash, InstancePath, InstanceObjects, TextureTags);
+	ImportHierarchy(ActorObject->GetDown(), DataObject->GetDown(), ActorElement, NewWorldTransformMatrix, InstancePath, InstanceObjects);
 
 	return ActorElement;
 }
 
-void FDatasmithC4DImporter::ImportHierarchy(melange::BaseObject* ActorObject, melange::BaseObject* DataObject, TSharedPtr<IDatasmithActorElement> ParentActor, const melange::Matrix& WorldTransformMatrix, TMap<FString, melange::PolygonObject*>* ClonerBaseChildrenHash, const FString& InstancePath, TArray<melange::BaseObject*>* InstanceObjects, const TArray<melange::TextureTag*>& TextureTags)
+void FDatasmithC4DImporter::ImportHierarchy(melange::BaseObject* ActorObject, melange::BaseObject* DataObject, TSharedPtr<IDatasmithActorElement> ParentActor, const melange::Matrix& WorldTransformMatrix, const FString& InstancePath, TArray<melange::BaseObject*>* InstanceObjects)
 {
 	while (ActorObject || DataObject)
 	{
@@ -2416,7 +2611,7 @@ void FDatasmithC4DImporter::ImportHierarchy(melange::BaseObject* ActorObject, me
 		if (!SkipObject)
 		{
 			FString DatasmithLabel = FDatasmithUtils::SanitizeObjectName(MelangeObjectName(ActorObject));
-			ImportObjectAndChildren(ActorObject, DataObject, ParentActor, WorldTransformMatrix, ClonerBaseChildrenHash, InstancePath, InstanceObjects, TextureTags, DatasmithLabel);
+			ImportObjectAndChildren(ActorObject, DataObject, ParentActor, WorldTransformMatrix, InstancePath, InstanceObjects, DatasmithLabel);
 		}
 
 		ActorObject = ActorObject->GetNext();
@@ -2424,7 +2619,7 @@ void FDatasmithC4DImporter::ImportHierarchy(melange::BaseObject* ActorObject, me
 	}
 }
 
-TSharedPtr<IDatasmithMeshElement> FDatasmithC4DImporter::ImportMesh(melange::PolygonObject* PolyObject, const FString& DatasmithMeshName, const FString& DatasmithLabel, const TArray<melange::TextureTag*>& TextureTags)
+TSharedPtr<IDatasmithMeshElement> FDatasmithC4DImporter::ImportMesh(melange::PolygonObject* PolyObject, const FString& DatasmithMeshName, const FString& DatasmithLabel)
 {
 	melange::Int32 PointCount = PolyObject->GetPointCount();
 	melange::Int32 PolygonCount = PolyObject->GetPolygonCount();
@@ -2441,26 +2636,36 @@ TSharedPtr<IDatasmithMeshElement> FDatasmithC4DImporter::ImportMesh(melange::Pol
 
 	// Collect all UV channels and material slot information for this PolygonObject
 	TArray<melange::ConstUVWHandle> UVWTagsData;
-	TMap<FString, melange::BaseSelect*> SelectionTags;
+	TArray<melange::BaseSelect*> SelectionTags;
+	SelectionTags.Add(nullptr); // The "unselected" group
+
 	for (melange::BaseTag* Tag = PolyObject->GetFirstTag(); Tag; Tag = Tag->GetNext())
 	{
 		melange::Int32 TagType = Tag->GetType();
 		if (TagType == Tuvw)
 		{
 			melange::UVWTag* UVWTag = static_cast<melange::UVWTag*>(Tag);
-			DatasmithC4DImportCheck(UVWTag->GetDataCount() == PolygonCount);
-			UVWTagsData.Add(UVWTag->GetDataAddressR());
+			if (UVWTag->GetDataCount() == PolygonCount)
+			{
+				UVWTagsData.Add(UVWTag->GetDataAddressR());
+			}
+			else
+			{
+				return TSharedPtr<IDatasmithMeshElement>();
+			}
 		}
 		else if (TagType == Tpolygonselection)
 		{
 			FString SelectionName = MelangeGetString(Tag, melange::POLYGONSELECTIONTAG_NAME);
 			if (!SelectionName.IsEmpty())
 			{
-				SelectionTags.Add(SelectionName, static_cast<melange::SelectionTag*>(Tag)->GetBaseSelect());
+				SelectionTags.Add(static_cast<melange::SelectionTag*>(Tag)->GetBaseSelect());
 			}
 		}
 	}
 
+	int32 NumSlots = SelectionTags.Num();
+
 	// Create MeshDescription
 	FMeshDescription MeshDescription;
 	DatasmithMeshHelper::PrepareAttributeForStaticMesh(MeshDescription);
@@ -2477,7 +2682,7 @@ TSharedPtr<IDatasmithMeshElement> FDatasmithC4DImporter::ImportMesh(melange::Pol
 	MeshDescription.ReserveNewVertexInstances(PolygonCount);
 	MeshDescription.ReserveNewEdges(PolygonCount);
 	MeshDescription.ReserveNewPolygons(PolygonCount);
-	MeshDescription.ReserveNewPolygonGroups(SelectionTags.Num() + 1);
+	MeshDescription.ReserveNewPolygonGroups(NumSlots);
 
 	// At least one UV set must exist.
 	int32 UVChannelCount = UVWTagsData.Num();
@@ -2493,21 +2698,12 @@ TSharedPtr<IDatasmithMeshElement> FDatasmithC4DImporter::ImportMesh(melange::Pol
 		VertexPositions[NewVertexID] = ConvertMelangePosition(Points[PointIndex]);
 	}
 
-	// Auxiliary stuff to help with polygon material assignment and material slots
-	int32 MaterialCounter = 0;
-	TMap<melange::TextureTag*, int32> TextureTagToMaterialSlot;
-	TMap<int32, FPolygonGroupID> MaterialSlotToPolygonGroup;
-	auto GetOrCreatePolygonGroupId = [&](int32 MaterialIndex) -> FPolygonGroupID
+	// Create one material slot per polygon selection tag (including the "unselected" group)
+	for (int32 SlotIndex = 0; SlotIndex < NumSlots; ++SlotIndex)
 	{
-		FPolygonGroupID& PolyGroupId = MaterialSlotToPolygonGroup.FindOrAdd(MaterialIndex);
-		if (PolyGroupId == FPolygonGroupID::Invalid)
-		{
-			PolyGroupId = MeshDescription.CreatePolygonGroup();
-			FName ImportedSlotName = DatasmithMeshHelper::DefaultSlotName(MaterialIndex);
-			PolygonGroupImportedMaterialSlotNames[PolyGroupId] = ImportedSlotName;
-		}
-		return PolyGroupId;
-	};
+		FPolygonGroupID PolyGroupId = MeshDescription.CreatePolygonGroup();
+		PolygonGroupImportedMaterialSlotNames[PolyGroupId] = DatasmithMeshHelper::DefaultSlotName(SlotIndex);
+	}
 
 	// Vertex indices in a quad or a triangle
 	TArray<int32> QuadIndexOffsets = { 0, 1, 3, 1, 2, 3 };
@@ -2608,41 +2804,20 @@ TSharedPtr<IDatasmithMeshElement> FDatasmithC4DImporter::ImportMesh(melange::Pol
 			}
 		}
 
-		// TextureTag
-		// Iterate backwards because the last valid texture tag is the one that is actually applied
-		melange::TextureTag* PolygonTextureTag = nullptr;
-		for (int32 TextureTagIndex = TextureTags.Num() - 1; TextureTagIndex >= 0; --TextureTagIndex)
+		// Find which selection tag (and so which material slot and polygon group) we belong to.
+		// Note that if we don't find any we end up in SlotIndex 0, which is the "unselected" group.
+		// Also note that we already receive just one texture tag per selection
+		int32 SlotIndex = NumSlots - 1;
+		for (; SlotIndex > 0; --SlotIndex)
 		{
-			melange::TextureTag* TextureTag = TextureTags[TextureTagIndex];
+			melange::BaseSelect* SelectionTag = SelectionTags[SlotIndex];
 
-			melange::BaseSelect** SelectionInMap = nullptr;
-			FString TextureSelectionTag = MelangeGetString(TextureTag, melange::TEXTURETAG_RESTRICTION);
-			if (!TextureSelectionTag.IsEmpty())
+			if (SelectionTag && SelectionTag->IsSelected(PolygonIndex))
 			{
-				SelectionInMap = SelectionTags.Find(TextureSelectionTag);
-			}
-
-			if (TextureSelectionTag.IsEmpty() || (SelectionInMap && (*SelectionInMap)->IsSelected(PolygonIndex)))
-			{
-				PolygonTextureTag = TextureTag;
-
-				// We just need the "last valid" one
 				break;
 			}
 		}
 
-		// MaterialIndex from TextureTag
-		int32 MaterialIndex = -1;
-		if (int32* FoundMaterialIndex = TextureTagToMaterialSlot.Find(PolygonTextureTag))
-		{
-			MaterialIndex = *FoundMaterialIndex;
-		}
-		else
-		{
-			MaterialIndex = MaterialCounter++;
-			TextureTagToMaterialSlot.Add(PolygonTextureTag, MaterialIndex);
-		}
-
 		// Create a triangle for each 3 vertex instance IDs we have
 		check(VertexInstances.Num() % 3 == 0);
 		for (int32 TriangleIndex = 0; TriangleIndex < VertexInstances.Num() / 3; ++TriangleIndex)
@@ -2664,8 +2839,7 @@ TSharedPtr<IDatasmithMeshElement> FDatasmithC4DImporter::ImportMesh(melange::Pol
 				continue; // this will leave holes...
 			}
 
-			const FPolygonGroupID PolygonGroupID = GetOrCreatePolygonGroupId(MaterialIndex);
-			const FPolygonID NewPolygonID = MeshDescription.CreatePolygon(PolygonGroupID, IDsCopy);
+			const FPolygonID NewPolygonID = MeshDescription.CreatePolygon(FPolygonGroupID(SlotIndex), IDsCopy);
 
 			// Fill in the polygon's Triangles - this won't actually do any polygon triangulation as we always give it triangles
 			MeshDescription.ComputePolygonTriangulation(NewPolygonID);
@@ -2685,26 +2859,7 @@ TSharedPtr<IDatasmithMeshElement> FDatasmithC4DImporter::ImportMesh(melange::Pol
 	TSharedRef<IDatasmithMeshElement> MeshElement = FDatasmithSceneFactory::CreateMesh(*DatasmithMeshName);
 	MeshElement->SetLabel(*DatasmithLabel);
 
-	// Create customized materials for all the used texture tags. This because each tag
-	// actually represents a material "instance", and might have different settings like texture tiling
-	for (TPair<melange::TextureTag*, int32>& Pair : TextureTagToMaterialSlot)
-	{
-		melange::TextureTag* Tag = Pair.Key;
-		int32 TargetSlot = Pair.Value;
-
-		FString CustomizedMaterialName;
-		melange::BaseList2D* TextureMaterial = Tag ? MelangeGetLink(Tag, melange::TEXTURETAG_MATERIAL) : nullptr;
-		if (TextureMaterial)
-		{
-			// This can also return an existing material without necessarily spawning a new instance
-			CustomizedMaterialName = CustomizeMaterial(GetMelangeBaseList2dID(TextureMaterial), DatasmithMeshName, Tag);
-		}
-
-		MeshElement->SetMaterial(*CustomizedMaterialName, TargetSlot);
-	}
-
 	MeshElementToMeshDescription.Add(&MeshElement.Get(), MoveTemp(MeshDescription));
-	PolygonObjectToMeshElement.Add(PolyObject, MeshElement);
 
 	DatasmithScene->AddMesh(MeshElement);
 	return MeshElement;
@@ -2730,7 +2885,7 @@ bool FDatasmithC4DImporter::OpenFile(const FString& InFilename)
 	}
 
 	C4dDocument = NewObj(BaseDocument);
-	if (!C4dDocument)
+	if (C4dDocument == nullptr)
 	{
 		return false;
 	}
@@ -2744,7 +2899,8 @@ bool FDatasmithC4DImporter::OpenFile(const FString& InFilename)
 
 	if (C4Dfile->Open(DOC_IDENT, TCHAR_TO_ANSI(*InFilename), FILEOPEN_READ))
 	{
-		bool bSuccess = C4dDocument->ReadObject(C4Dfile, true);
+		// Extra nullptr check because static analysis tool doesn't understand that this was already checked
+		bool bSuccess = C4dDocument != nullptr && C4dDocument->ReadObject(C4Dfile, true);
 
 		int64 LastPos = static_cast<int64>(C4Dfile->GetPosition());
 		int64 Length = static_cast<int64>(C4Dfile->GetLength());
@@ -2773,28 +2929,44 @@ bool FDatasmithC4DImporter::OpenFile(const FString& InFilename)
 	C4Dfile->Close();
 	DeleteObj(C4Dfile);
 
+	if (C4dDocument != nullptr && !C4dDocument->HasCaches())
+	{
+		FText ErrorMsg = FText::Format(LOCTEXT("C4DNotSavedForMelange", "The file '{0}' was not saved for Melange."), FText::FromString(C4dDocumentFilename));
+		UE_LOG(LogDatasmithC4DImport, Warning, TEXT("%s"), *ErrorMsg.ToString());
+
+		FNotificationInfo NotificationInfo(ErrorMsg);
+		NotificationInfo.ExpireDuration = 5.0f;
+		FSlateNotificationManager::Get().AddNotification(NotificationInfo);
+	}
+
 	return true;
 }
 
 melange::BaseObject* FDatasmithC4DImporter::FindMelangeObject(const FString& SearchObjectID, melange::BaseObject* Object)
 {
-	while (Object)
+	melange::BaseObject* FoundObject = nullptr;
+
+	while (Object && !FoundObject)
 	{
-		if (MelangeObjectID(Object) == SearchObjectID)
+		if (TOptional<FString> ObjectID = MelangeObjectID(Object))
 		{
-			return Object;
+			if (ObjectID.GetValue() == SearchObjectID)
+			{
+				FoundObject = Object;
+				break;
+			}
 		}
-
-		melange::BaseObject* FoundObject = FindMelangeObject(SearchObjectID, Object->GetDown());
-		if (FoundObject)
+		else
 		{
-			return FoundObject;
+			//The object is invalid or we could not find its ID.
+			break;
 		}
 
+		FoundObject = FindMelangeObject(SearchObjectID, Object->GetDown());
 		Object = Object->GetNext();
 	}
 
-	return nullptr;
+	return FoundObject;
 }
 
 melange::BaseObject* FDatasmithC4DImporter::GoToMelangeHierarchyPosition(melange::BaseObject* Object, const FString& HierarchyPosition)
@@ -2845,16 +3017,17 @@ bool FDatasmithC4DImporter::ProcessScene()
 
 	// Materials
 	ImportedTextures.Empty();
-	ImportMaterialHierarchy(C4dDocument->GetFirstMaterial());
+	if (!ImportMaterialHierarchy(C4dDocument->GetFirstMaterial()))
+	{
+		return false;
+	}
 	ImportedTextures.Empty();
 
 	// Actors
-	ActorMetadata.Empty();
 	// Need a RootActor for RemoveEmptyActors and to make AddChildActor agnostic to actor hiearchy level
 	TSharedPtr<IDatasmithActorElement> RootActor = FDatasmithSceneFactory::CreateActor(TEXT("RootActor"));
 	DatasmithScene->AddActor(RootActor);
-	TArray<melange::TextureTag*> TextureTags;
-	ImportHierarchy(C4dDocument->GetFirstObject(), C4dDocument->GetFirstObject(), RootActor, melange::Matrix(), nullptr, "", nullptr, TextureTags);
+	ImportHierarchy(C4dDocument->GetFirstObject(), C4dDocument->GetFirstObject(), RootActor, melange::Matrix(), "", nullptr);
 
 	// Animations
 	LevelSequence = FDatasmithSceneFactory::CreateLevelSequence(DatasmithScene->GetName());
@@ -2863,8 +3036,8 @@ bool FDatasmithC4DImporter::ProcessScene()
 	ImportActorHierarchyAnimations(RootActor);
 
 	// Processing
-	KeepParentsOfAnimatedNodes(RootActor, NamesOfAnimatedActors);
-	RemoveEmptyActors();
+	FC4DImporterImpl::KeepParentsOfAnimatedNodes(RootActor, NamesOfAnimatedActors);
+	FC4DImporterImpl::RemoveEmptyActors(DatasmithScene, NamesOfCameraTargetActors, NamesOfAnimatedActors);
 	DatasmithScene->RemoveActor(RootActor, EDatasmithActorRemovalRule::KeepChildrenAndKeepRelativeTransform);
 
 	if (Options->bExportToUDatasmith)
diff --git a/Engine/Plugins/Enterprise/DatasmithC4DImporter/Source/DatasmithC4DTranslator/Private/DatasmithC4DImporter.h b/Engine/Plugins/Enterprise/DatasmithC4DImporter/Source/DatasmithC4DTranslator/Private/DatasmithC4DImporter.h
--- a/Engine/Plugins/Enterprise/DatasmithC4DImporter/Source/DatasmithC4DTranslator/Private/DatasmithC4DImporter.h
+++ b/Engine/Plugins/Enterprise/DatasmithC4DImporter/Source/DatasmithC4DTranslator/Private/DatasmithC4DImporter.h
@@ -22,7 +22,6 @@ class IDatasmithLightActorElement;
 class IDatasmithMasterMaterialElement;
 class IDatasmithMeshActorElement;
 class IDatasmithMeshElement;
-class IDatasmithMetaDataElement;
 class IDatasmithScene;
 class IDatasmithTextureElement;
 class UDatasmithC4DImportOptions;
@@ -58,7 +57,7 @@ public:
 	melange::BaseObject* GetBestMelangeCache(melange::BaseObject* Object);
 
 	/** Generates a unique identifier string for a melange object based on the object's name and its position in the hierarchy */
-	FString MelangeObjectID(melange::BaseObject* Object);
+	TOptional<FString> MelangeObjectID(melange::BaseObject* Object);
 
 	/** Searches the melange object hierarchy for a melange::BaseObject that has a MelangeObjectID equal to SearchObjectID */
 	melange::BaseObject* FindMelangeObject(const FString& SearchObjectID, melange::BaseObject* Object);
@@ -78,20 +77,22 @@ public:
 	/**
 	 * Import melange objects into Datasmith elements. Assets like meshes, materials and textures are added to the Datasmith scene directly,
 	 * while actors are merely returned and must be added as children to scene actors or added to the scene manually.
+	 *
+	 * Returns a null TSharedPtr if an error occurred during the import process.
 	 */
 	TSharedPtr<IDatasmithActorElement> ImportNullActor(melange::BaseObject* Object, const FString& DatasmithName, const FString& DatasmithLabel);
 	TSharedPtr<IDatasmithLightActorElement> ImportLight(melange::BaseObject* LightObject, const FString& DatasmithName, const FString& DatasmithLabel);
 	TSharedPtr<IDatasmithCameraActorElement> ImportCamera(melange::BaseObject* CameraObject, const FString& DatasmithName, const FString& DatasmithLabel);
-	TSharedPtr<IDatasmithMeshActorElement> ImportPolygon(melange::PolygonObject* PolyObject, TMap<FString, melange::PolygonObject*>* ClonerBaseChildrenHash, const FString& DatasmithName, const FString& DatasmithLabel, const TArray<melange::TextureTag*>& TextureTags);
+	TSharedPtr<IDatasmithMeshActorElement> ImportPolygon(melange::PolygonObject* PolyObject, const FString& DatasmithName, const FString& DatasmithLabel, const TArray<melange::TextureTag*>& TextureTags);
 	TSharedPtr<IDatasmithMasterMaterialElement> ImportMaterial(melange::Material* InC4DMaterialPtr);
 	TSharedPtr<IDatasmithTextureElement> ImportTexture(const FString& TexturePath, EDatasmithTextureMode TextureMode);
-	TSharedPtr<IDatasmithMeshElement> ImportMesh(melange::PolygonObject* PolyObject, const FString& DatasmithMeshName, const FString& DatasmithLabel, const TArray<melange::TextureTag*>& TextureTags);
+	TSharedPtr<IDatasmithMeshElement> ImportMesh(melange::PolygonObject* PolyObject, const FString& DatasmithMeshName, const FString& DatasmithLabel);
 
 	/** Parses the spline and its cache into SplineCurves so that it can be used as paths for animation later */
 	void ImportSpline(melange::SplineObject* ActorObject);
 
 	/** Traverse the melange material hierarchy contained in the c4d file and import each into IDatasmithMasterMaterialElements */
-	void ImportMaterialHierarchy(melange::BaseMaterial* InC4DMaterialPtr);
+	bool ImportMaterialHierarchy(melange::BaseMaterial* InC4DMaterialPtr);
 
 	/** Uses ActorElementToC4DObject to find the corresponding melange object for ActorElement and adds all of its animations to LevelSequence */
 	void ImportAnimations(TSharedPtr<IDatasmithActorElement> ActorElement);
@@ -109,20 +110,22 @@ public:
 	 */
 	FString CustomizeMaterial(const FString& InMaterialID, const FString& InMeshID, melange::TextureTag* InTextureTag);
 
+	/**
+	 * Creates customized materials if necessary, and returns a map from material slot indices to material names
+	 */
+	TMap<int32, FString> GetCustomizedMaterialAssignment(const FString& DatasmithMeshName, const TArray<melange::TextureTag*>& TextureTags);
+
 	/** Imports a melange actor, which might involve parsing another small hierarchy of subnodes and deformers*/
-	TSharedPtr<IDatasmithActorElement> ImportObjectAndChildren(melange::BaseObject* ActorObject, melange::BaseObject* DataObject, TSharedPtr<IDatasmithActorElement> ParentActor, const melange::Matrix& WorldTransformMatrix, TMap<FString, melange::PolygonObject*>* ClonerBaseChildrenHash, const FString& InstancePath, TArray<melange::BaseObject*>* InstanceObjects, TArray<melange::TextureTag*> TextureTags, const FString& DatasmithLabel);
+	TSharedPtr<IDatasmithActorElement> ImportObjectAndChildren(melange::BaseObject* ActorObject, melange::BaseObject* DataObject, TSharedPtr<IDatasmithActorElement> ParentActor, const melange::Matrix& WorldTransformMatrix, const FString& InstancePath, TArray<melange::BaseObject*>* InstanceObjects, const FString& DatasmithLabel);
 
 	/** Traverse the melange actor hierarchy importing all nodes */
-	void ImportHierarchy(melange::BaseObject* ActorObject, melange::BaseObject* DataObject, TSharedPtr<IDatasmithActorElement> ParentActor, const melange::Matrix& WorldTransformMatrix, TMap<FString, melange::PolygonObject*>* ClonerBaseChildrenHash, const FString& InstancePath, TArray<melange::BaseObject*>* InstanceObjects, const TArray<melange::TextureTag*>& TextureTags);
+	void ImportHierarchy(melange::BaseObject* ActorObject, melange::BaseObject* DataObject, TSharedPtr<IDatasmithActorElement> ParentActor, const melange::Matrix& WorldTransformMatrix, const FString& InstancePath, TArray<melange::BaseObject*>* InstanceObjects);
 
 	/**
 	 * Adds Actor as a child of ParentActor using the corresponding WorldTransformMatrix. Object is the corresponding melange Object to Actor.
 	 * This will also do the necessary coordinate system conversions between melange and Unreal.
 	 */
-	void AddChildActor(melange::BaseObject* Object, TSharedPtr<IDatasmithActorElement> ParentActor, melange::Matrix WorldTransformMatrix, const TSharedPtr<IDatasmithActorElement>& Actor);
-
-	/** Creates a new IDatasmithMetaDataElement for the IDatasmithActorElement, adds it to the scene and returns it */
-	TSharedPtr<IDatasmithMetaDataElement> CreateMetadataForActor(const IDatasmithActorElement& Actor);
+	bool AddChildActor(melange::BaseObject* Object, TSharedPtr<IDatasmithActorElement> ParentActor, melange::Matrix WorldTransformMatrix, const TSharedPtr<IDatasmithActorElement>& Actor);
 
 	void GetGeometriesForMeshElementAndRelease(const TSharedRef<IDatasmithMeshElement> MeshElement, TArray<FMeshDescription>& OutMeshDescriptions);
 	TSharedPtr<IDatasmithLevelSequenceElement> GetLevelSequence() { return LevelSequence; }
@@ -131,6 +134,9 @@ public:
 	FString C4dDocumentFilename;
 
 private:
+	/** Returns the TextureTags that should affect this object. May check parent objects, so relies on CachesOriginalObject */
+	TArray<melange::TextureTag*> GetActiveTextureTags(const melange::BaseObject* Object);
+
 	/** Removes from Context->Scene all empty actors that have a single child */
 	void RemoveEmptyActors();
 
@@ -143,8 +149,8 @@ private:
 	/** Storage of created materials used by CustomizeMaterial to create new "material instances" */
 	TMap<FString, TSharedPtr<IDatasmithMasterMaterialElement>> MaterialNameToMaterialElement;
 
-	/** Cache to prevent us from importing the mesh from a melange::PolygonObject more than once */
-	TMap<melange::PolygonObject*, TSharedRef<IDatasmithMeshElement>> PolygonObjectToMeshElement;
+	/** Cache meshes by hash to promote reusing StaticMeshes */
+	TMap<FString, TSharedRef<IDatasmithMeshElement>> PolygonHashToMeshElement;
 
 	/** Cache to prevent us from importing the same texture in the same mode more than once (mode is encoded in the FString as well) */
 	TMap<FString, TSharedPtr<IDatasmithTextureElement>> ImportedTextures;
@@ -158,9 +164,6 @@ private:
 	/** Keeps track of the owners of every melange cache object so we can climb the hierarchy upwards */
 	TMap<melange::BaseObject*, melange::BaseObject*> CachesOriginalObject;
 
-	/** Caches all the existing actor metadata elements we created for each IDatasmithActorElement in our Datasmith scene */
-	TMap<const IDatasmithActorElement*, TSharedPtr<IDatasmithMetaDataElement>> ActorMetadata;
-
 	/** Stores all FCraneCameraAttributes for each camera */
 	TMap<melange::BaseObject*, TSharedRef<FCraneCameraAttributes>> CraneCameraToAttributes;
 
diff --git a/Engine/Plugins/Enterprise/DatasmithC4DImporter/Source/DatasmithC4DTranslator/Private/DatasmithC4DTranslatorModule.cpp b/Engine/Plugins/Enterprise/DatasmithC4DImporter/Source/DatasmithC4DTranslator/Private/DatasmithC4DTranslatorModule.cpp
--- a/Engine/Plugins/Enterprise/DatasmithC4DImporter/Source/DatasmithC4DTranslator/Private/DatasmithC4DTranslatorModule.cpp
+++ b/Engine/Plugins/Enterprise/DatasmithC4DImporter/Source/DatasmithC4DTranslator/Private/DatasmithC4DTranslatorModule.cpp
@@ -6,10 +6,6 @@
 
 #include "DatasmithC4DTranslator.h"
 #include "DatasmithC4DImporterMaterialSelector.h"
-#include "DatasmithImporterEditorSettings.h"
-#include "DatasmithImporterHelper.h"
-#include "DatasmithImporterModule.h"
-#include "UI/DatasmithUIManager.h"
 
 #include "MasterMaterials/DatasmithMasterMaterialManager.h"
 
diff --git a/Engine/Plugins/Enterprise/DatasmithC4DImporter/Source/DatasmithC4DTranslator/Private/DatasmithC4DUtils.cpp b/Engine/Plugins/Enterprise/DatasmithC4DImporter/Source/DatasmithC4DTranslator/Private/DatasmithC4DUtils.cpp
--- a/Engine/Plugins/Enterprise/DatasmithC4DImporter/Source/DatasmithC4DTranslator/Private/DatasmithC4DUtils.cpp
+++ b/Engine/Plugins/Enterprise/DatasmithC4DImporter/Source/DatasmithC4DTranslator/Private/DatasmithC4DUtils.cpp
@@ -10,8 +10,6 @@
 #include "c4d_file.h"
 #include "DatasmithC4DMelangeSDKLeaveGuard.h"
 
-#include "DatasmithC4DImportException.h"
-
 #include "HAL/FileManager.h"
 #include "Misc/Paths.h"
 #include "Misc/SecureHash.h"
@@ -420,18 +418,24 @@ FString MelangeParameterValueToString(melange::BaseList2D* Object, melange::Int3
 	return TEXT("");
 }
 
-FString GetMelangeBaseList2dID(melange::BaseList2D* BaseList)
+TOptional<FString> GetMelangeBaseList2dID(melange::BaseList2D* BaseList)
 {
+	TOptional<FString> MelangeID;
 	if (BaseList)
 	{
-		DatasmithC4DImportCheck(BaseList->GetUniqueIDCount() > 0);
-		melange::Int32 AppId;
-		const melange::Char* IdData;
-		melange::Int Bytes;
-		DatasmithC4DImportCheck(BaseList->GetUniqueIDIndex(0, AppId, IdData, Bytes));
-		return BytesToHex(reinterpret_cast<const uint8*>(&AppId), sizeof(AppId)) + "_" + BytesToHex(reinterpret_cast<const uint8*>(IdData), static_cast<int32>(Bytes));
+		if (BaseList->GetUniqueIDCount() > 0)
+		{
+			melange::Int32 AppId;
+			const melange::Char* IdData;
+			melange::Int Bytes;
+			if (BaseList->GetUniqueIDIndex(0, AppId, IdData, Bytes))
+			{
+				MelangeID.Emplace(BytesToHex(reinterpret_cast<const uint8*>(&AppId), sizeof(AppId)) + "_" + BytesToHex(reinterpret_cast<const uint8*>(IdData), static_cast<int32>(Bytes)));
+			}
+		}
 	}
-	return TEXT("Invalid object");
+	
+	return MelangeID;
 }
 
 #endif
\ No newline at end of file
diff --git a/Engine/Plugins/Enterprise/DatasmithC4DImporter/Source/DatasmithC4DTranslator/Private/DatasmithC4DUtils.h b/Engine/Plugins/Enterprise/DatasmithC4DImporter/Source/DatasmithC4DTranslator/Private/DatasmithC4DUtils.h
--- a/Engine/Plugins/Enterprise/DatasmithC4DImporter/Source/DatasmithC4DTranslator/Private/DatasmithC4DUtils.h
+++ b/Engine/Plugins/Enterprise/DatasmithC4DImporter/Source/DatasmithC4DTranslator/Private/DatasmithC4DUtils.h
@@ -165,6 +165,6 @@ FString MelangeParameterValueToString(melange::BaseList2D* Object, melange::Int3
 /**
  * Returns the full melange ID for the BaseList2D argument as an FString, which will include AppId
  */
-FString GetMelangeBaseList2dID(melange::BaseList2D* BaseList);
+TOptional<FString> GetMelangeBaseList2dID(melange::BaseList2D* BaseList);
 
 #endif
\ No newline at end of file
diff --git a/Engine/Plugins/Enterprise/DatasmithCADImporter/DatasmithCADImporter.uplugin b/Engine/Plugins/Enterprise/DatasmithCADImporter/DatasmithCADImporter.uplugin
--- a/Engine/Plugins/Enterprise/DatasmithCADImporter/DatasmithCADImporter.uplugin
+++ b/Engine/Plugins/Enterprise/DatasmithCADImporter/DatasmithCADImporter.uplugin
@@ -80,6 +80,14 @@
 			"WhitelistPlatforms": [
 				"Win64"
 			]
+		},
+		{
+			"Name": "DatasmithCoreTechParametricSurfaceData",
+			"Type": "Editor",
+			"LoadingPhase": "Default",
+			"WhitelistPlatforms": [
+				"Win64"
+			]
 		}
 	],
 	"Plugins": [
