C4D: Copy cls 10881806 and 10449975

Fixes actors being discarded if they have zero scaling;
Restore toast message when not saved for melange;

diff --git a/Engine/Plugins/Enterprise/DatasmithC4DImporter/Source/DatasmithC4DTranslator/Private/DatasmithC4DImporter.cpp b/Engine/Plugins/Enterprise/DatasmithC4DImporter/Source/DatasmithC4DTranslator/Private/DatasmithC4DImporter.cpp
--- a/Engine/Plugins/Enterprise/DatasmithC4DImporter/Source/DatasmithC4DTranslator/Private/DatasmithC4DImporter.cpp
+++ b/Engine/Plugins/Enterprise/DatasmithC4DImporter/Source/DatasmithC4DTranslator/Private/DatasmithC4DImporter.cpp
@@ -29,6 +29,7 @@
 #include "Engine/StaticMesh.h"
 #include "Engine/StaticMeshActor.h"
 #include "Engine/World.h"
+#include "Framework/Notifications/NotificationManager.h"
 #include "HAL/Event.h"
 #include "HAL/FileManager.h"
 #include "HAL/Runnable.h"
@@ -45,6 +46,7 @@
 #include "PackageTools.h"
 #include "RawMesh.h"
 #include "StaticMeshAttributes.h"
+#include "Widgets/Notifications/SNotificationList.h"
 
 #include "ImathMatrixAlgo.h"
 
@@ -708,14 +710,17 @@ void FDatasmithC4DImporter::AddChildActor(melange::BaseObject* Object, TSharedPt
 		0.0,            0.0,             0.0, 1.0);
 
 	// Remove any scaling from the matrix and get the scale vector that was initially present.
-	Imath::Vec3<float> Scale;
-	Imath::Vec3<float> Shear;
+	Imath::Vec3<float> Scale(1.0f, 1.0f, 1.0f);
+	Imath::Vec3<float> Shear(0.0f, 0.0f, 0.0f);
 	bool bExtracted = Imath::extractAndRemoveScalingAndShear<float>(Matrix, Scale, Shear, false);
 	if (!bExtracted)
 	{
-		// TODO: Append a message to the build summary.
-		FString Msg = FString::Printf(TEXT("WARNING: Actor %ls (%ls) has some zero scaling"), Actor->GetName(), Actor->GetLabel());
-		return;
+		UE_LOG(LogDatasmithC4DImport, Warning, TEXT("Actor %ls (%ls) has some zero scaling"), Actor->GetName(), Actor->GetLabel());
+
+		// extractAndRemoveScalingAndShear may have partially written to these vectors, so we need to
+		// reset them here to make sure they're valid for code below
+		Scale = Imath::Vec3<float>(1.0f, 1.0f, 1.0f);
+		Shear = Imath::Vec3<float>(0.0f, 0.0f, 0.0f);
 	}
 
 	// Initialize a rotation quaternion with the rotation matrix.
@@ -2773,6 +2778,16 @@ bool FDatasmithC4DImporter::OpenFile(const FString& InFilename)
 	C4Dfile->Close();
 	DeleteObj(C4Dfile);
 
+	if (C4dDocument && !C4dDocument->HasCaches())
+	{
+		FText ErrorMsg = FText::Format(LOCTEXT("C4DNotSavedForMelange", "The file '{0}' was not saved for Melange."), FText::FromString(C4dDocumentFilename));
+		UE_LOG(LogDatasmithC4DImport, Warning, TEXT("%s"), *ErrorMsg.ToString());
+
+		FNotificationInfo NotificationInfo(ErrorMsg);
+		NotificationInfo.ExpireDuration = 5.0f;
+		FSlateNotificationManager::Get().AddNotification(NotificationInfo);
+	}
+
 	return true;
 }
 
