C4D: Fix minor static analysis warning

diff --git a/Engine/Plugins/Enterprise/DatasmithC4DImporter/Source/DatasmithC4DTranslator/Private/DatasmithC4DImporter.cpp b/Engine/Plugins/Enterprise/DatasmithC4DImporter/Source/DatasmithC4DTranslator/Private/DatasmithC4DImporter.cpp
--- a/Engine/Plugins/Enterprise/DatasmithC4DImporter/Source/DatasmithC4DTranslator/Private/DatasmithC4DImporter.cpp
+++ b/Engine/Plugins/Enterprise/DatasmithC4DImporter/Source/DatasmithC4DTranslator/Private/DatasmithC4DImporter.cpp
@@ -2735,7 +2735,7 @@ bool FDatasmithC4DImporter::OpenFile(const FString& InFilename)
 	}
 
 	C4dDocument = NewObj(BaseDocument);
-	if (!C4dDocument)
+	if (C4dDocument == nullptr)
 	{
 		return false;
 	}
@@ -2749,7 +2749,8 @@ bool FDatasmithC4DImporter::OpenFile(const FString& InFilename)
 
 	if (C4Dfile->Open(DOC_IDENT, TCHAR_TO_ANSI(*InFilename), FILEOPEN_READ))
 	{
-		bool bSuccess = C4dDocument->ReadObject(C4Dfile, true);
+		// Extra nullptr check because static analysis tool doesn't understand that this was already checked
+		bool bSuccess = C4dDocument != nullptr && C4dDocument->ReadObject(C4Dfile, true);
 
 		int64 LastPos = static_cast<int64>(C4Dfile->GetPosition());
 		int64 Length = static_cast<int64>(C4Dfile->GetLength());
@@ -2778,7 +2779,7 @@ bool FDatasmithC4DImporter::OpenFile(const FString& InFilename)
 	C4Dfile->Close();
 	DeleteObj(C4Dfile);
 
-	if (C4dDocument && !C4dDocument->HasCaches())
+	if (C4dDocument != nullptr && !C4dDocument->HasCaches())
 	{
 		FText ErrorMsg = FText::Format(LOCTEXT("C4DNotSavedForMelange", "The file '{0}' was not saved for Melange."), FText::FromString(C4dDocumentFilename));
 		UE_LOG(LogDatasmithC4DImport, Warning, TEXT("%s"), *ErrorMsg.ToString());
