C4D: Merge "Fix texture tags not being propagated from parent actors" from CL 11208396

diff --git a/Engine/Plugins/Enterprise/DatasmithC4DImporter/Source/DatasmithC4DTranslator/Private/DatasmithC4DImporter.cpp b/Engine/Plugins/Enterprise/DatasmithC4DImporter/Source/DatasmithC4DTranslator/Private/DatasmithC4DImporter.cpp
--- a/Engine/Plugins/Enterprise/DatasmithC4DImporter/Source/DatasmithC4DTranslator/Private/DatasmithC4DImporter.cpp
+++ b/Engine/Plugins/Enterprise/DatasmithC4DImporter/Source/DatasmithC4DTranslator/Private/DatasmithC4DImporter.cpp
@@ -2287,41 +2287,23 @@ void FDatasmithC4DImporter::ImportActorHierarchyAnimations(TSharedPtr<IDatasmith
 	}
 }
 
-TArray<melange::TextureTag*> FDatasmithC4DImporter::GetActiveTextureTags(const melange::BaseObject* Object)
+TArray<melange::TextureTag*> FDatasmithC4DImporter::GetActiveTextureTags(const melange::BaseObject* Object, const TArray<melange::TextureTag*>& OrderedTextureTags)
 {
 	if (!Object)
 	{
 		return {};
 	}
 
-	// If we're the cache of an object, then our texture tags will be stored on that object.
-	// This repeats recursively, with only the first 'non-cache' BaseObject actually containing the TextureTags.
-	// Most PolygonObjects are not caches, but if a polygon has direct children then this will happen.
-	const melange::BaseObject* ObjectToSearch = Object;
-	while (const melange::BaseObject* Parent = CachesOriginalObject.FindRef(ObjectToSearch))
-	{
-		ObjectToSearch = Parent;
-	}
-
-	if (ObjectToSearch == nullptr)
-	{
-		return {};
-	}
-
-	TArray<melange::TextureTag*> OrderedTextureTags;
 	TArray<melange::BaseSelect*> OrderedSelectionTags;
 	OrderedSelectionTags.Add(nullptr); // "unselected" group
 
 	TMap<FString, melange::BaseSelect*> SelectionTagsByName;
 
-	for (melange::BaseTag* Tag = const_cast<melange::BaseObject*>(ObjectToSearch)->GetFirstTag(); Tag; Tag = Tag->GetNext())
+	// Fetch selection tags, which only affect this polygon
+	// The texture tags are fetched when moving down the hierarchy, as texture tags on parents also affect children
+	for (melange::BaseTag* Tag = const_cast<melange::BaseObject*>(Object)->GetFirstTag(); Tag; Tag = Tag->GetNext())
 	{
-		melange::Int32 TagType = Tag->GetType();
-		if (TagType == Ttexture)
-		{
-			OrderedTextureTags.Add(static_cast<melange::TextureTag*>(Tag));
-		}
-		else if (TagType == Tpolygonselection)
+		if (Tag->GetType() == Tpolygonselection)
 		{
 			FString SelectionName = MelangeGetString(Tag, melange::POLYGONSELECTIONTAG_NAME);
 			if (!SelectionName.IsEmpty())
@@ -2369,7 +2351,7 @@ TArray<melange::TextureTag*> FDatasmithC4DImporter::GetActiveTextureTags(const m
 	return OrderedActiveTextureTags;
 }
 
-TSharedPtr<IDatasmithActorElement> FDatasmithC4DImporter::ImportObjectAndChildren(melange::BaseObject* ActorObject, melange::BaseObject* DataObject, TSharedPtr<IDatasmithActorElement> ParentActor, const melange::Matrix& WorldTransformMatrix, const FString& InstancePath, TArray<melange::BaseObject*>* InstanceObjects, const FString& DatasmithLabel)
+TSharedPtr<IDatasmithActorElement> FDatasmithC4DImporter::ImportObjectAndChildren(melange::BaseObject* ActorObject, melange::BaseObject* DataObject, TSharedPtr<IDatasmithActorElement> ParentActor, const melange::Matrix& WorldTransformMatrix, const FString& InstancePath, TArray<melange::BaseObject*>* InstanceObjects, const FString& DatasmithLabel, const TArray<melange::TextureTag*>& TextureTags)
 {
 	TSharedPtr<IDatasmithActorElement> ActorElement;
 	melange::Int32 ObjectType = DataObject->GetType();
@@ -2460,7 +2442,7 @@ TSharedPtr<IDatasmithActorElement> FDatasmithC4DImporter::ImportObjectAndChildre
 
 					if (TOptional<FString> ObjectID = MelangeObjectID(DataObject))
 					{
-						return ImportObjectAndChildren(ActorObject, ActorToImport, ParentActor, WorldTransformMatrix, ObjectID.GetValue() + InstancePath, &CurentInstanceObjects, DatasmithLabel);
+						return ImportObjectAndChildren(ActorObject, ActorToImport, ParentActor, WorldTransformMatrix, ObjectID.GetValue() + InstancePath, &CurentInstanceObjects, DatasmithLabel, TextureTags);
 					}
 				}
 
@@ -2503,7 +2485,7 @@ TSharedPtr<IDatasmithActorElement> FDatasmithC4DImporter::ImportObjectAndChildre
 					ActorElement = ImportNullActor(ActorObject, DatasmithName.GetValue(), DatasmithLabel);
 					if (DataCache != nullptr && DataCache->GetType() == Onull && AddChildActor(ActorObject, ParentActor, NewWorldTransformMatrix, ActorElement))
 					{
-						ImportHierarchy(ActorCache->GetDown(), DataCache->GetDown(), ActorElement, NewWorldTransformMatrix, InstancePath, nullptr);
+						ImportHierarchy(ActorCache->GetDown(), DataCache->GetDown(), ActorElement, NewWorldTransformMatrix, InstancePath, nullptr, TextureTags);
 						return ActorElement;
 					}
 
@@ -2519,7 +2501,7 @@ TSharedPtr<IDatasmithActorElement> FDatasmithC4DImporter::ImportObjectAndChildre
 			case Oboole:
 				ActorElement = ImportNullActor(ActorObject, DatasmithName.GetValue() +"0"/*to be different than the cache root*/, DatasmithLabel);
 				if (DataCache != nullptr && AddChildActor(ActorObject, ParentActor, NewWorldTransformMatrix, ActorElement)
-					&& ImportObjectAndChildren(ActorCache, DataCache, ActorElement, NewWorldTransformMatrix, InstancePath, nullptr, DatasmithLabel))
+					&& ImportObjectAndChildren(ActorCache, DataCache, ActorElement, NewWorldTransformMatrix, InstancePath, nullptr, DatasmithLabel, TextureTags))
 				{
 					return ActorElement;
 				}
@@ -2541,14 +2523,15 @@ TSharedPtr<IDatasmithActorElement> FDatasmithC4DImporter::ImportObjectAndChildre
 
 			if (bSuccess && bImportCache && ActorCache)
 			{
-				ActorElement = ImportObjectAndChildren(ActorCache, DataCache, nullptr, ActorCache->GetMg(), InstancePath, nullptr, DatasmithLabel);
+				ActorElement = ImportObjectAndChildren(ActorCache, DataCache, nullptr, ActorCache->GetMg(), InstancePath, nullptr, DatasmithLabel, TextureTags);
 			}
 			else if (ObjectType == Opolygon)
 			{
 				melange::PolygonObject* PolygonObject = static_cast<melange::PolygonObject*>(DataObject);
 				if (Options->bImportEmptyMesh || PolygonObject->GetPolygonCount() > 0)
 				{
-					if (TSharedPtr<IDatasmithMeshActorElement> MeshActorElement = ImportPolygon(PolygonObject, DatasmithName.GetValue(), DatasmithLabel, GetActiveTextureTags(PolygonObject)))
+					TArray<melange::TextureTag*> ActiveTextureTags = GetActiveTextureTags(PolygonObject, TextureTags);
+					if (TSharedPtr<IDatasmithMeshActorElement> MeshActorElement = ImportPolygon(PolygonObject, DatasmithName.GetValue(), DatasmithLabel, ActiveTextureTags))
 					{
 						ActorElement = MeshActorElement;
 					}
@@ -2603,12 +2586,12 @@ TSharedPtr<IDatasmithActorElement> FDatasmithC4DImporter::ImportObjectAndChildre
 		ActorElement->SetLayer(*TargetLayerName);
 	}
 
-	ImportHierarchy(ActorObject->GetDown(), DataObject->GetDown(), ActorElement, NewWorldTransformMatrix, InstancePath, InstanceObjects);
+	ImportHierarchy(ActorObject->GetDown(), DataObject->GetDown(), ActorElement, NewWorldTransformMatrix, InstancePath, InstanceObjects, TextureTags);
 
 	return ActorElement;
 }
 
-void FDatasmithC4DImporter::ImportHierarchy(melange::BaseObject* ActorObject, melange::BaseObject* DataObject, TSharedPtr<IDatasmithActorElement> ParentActor, const melange::Matrix& WorldTransformMatrix, const FString& InstancePath, TArray<melange::BaseObject*>* InstanceObjects)
+void FDatasmithC4DImporter::ImportHierarchy(melange::BaseObject* ActorObject, melange::BaseObject* DataObject, TSharedPtr<IDatasmithActorElement> ParentActor, const melange::Matrix& WorldTransformMatrix, const FString& InstancePath, TArray<melange::BaseObject*>* InstanceObjects, const TArray<melange::TextureTag*>& TextureTags)
 {
 	while (ActorObject || DataObject)
 	{
@@ -2621,10 +2604,14 @@ void FDatasmithC4DImporter::ImportHierarchy(melange::BaseObject* ActorObject, me
 			ActorObject = DataObject;
 		}
 
+		// Reset this for every child as texture tags only propagate down, not between siblings
+		TArray<melange::TextureTag*> TextureTagsDown = TextureTags;
+
 		bool SkipObject = false;
 		for (melange::BaseTag* Tag = ActorObject->GetFirstTag(); Tag; Tag = Tag->GetNext())
 		{
-			if (Tag->GetType() == Tannotation)
+			melange::Int32 TagType = Tag->GetType();
+			if (TagType == Tannotation)
 			{
 				FString AnnotationLabel = MelangeGetString(Tag, 10014);
 				if (AnnotationLabel.Compare("EXCLUDE", ESearchCase::IgnoreCase) == 0)
@@ -2633,12 +2620,16 @@ void FDatasmithC4DImporter::ImportHierarchy(melange::BaseObject* ActorObject, me
 					break;
 				}
 			}
+			else if (TagType == Ttexture)
+			{
+				TextureTagsDown.Add(static_cast<melange::TextureTag*>(Tag));
+			}
 		}
 
 		if (!SkipObject)
 		{
 			FString DatasmithLabel = FDatasmithUtils::SanitizeObjectName(MelangeObjectName(ActorObject));
-			ImportObjectAndChildren(ActorObject, DataObject, ParentActor, WorldTransformMatrix, InstancePath, InstanceObjects, DatasmithLabel);
+			ImportObjectAndChildren(ActorObject, DataObject, ParentActor, WorldTransformMatrix, InstancePath, InstanceObjects, DatasmithLabel, TextureTagsDown);
 		}
 
 		ActorObject = ActorObject->GetNext();
@@ -3054,7 +3045,8 @@ bool FDatasmithC4DImporter::ProcessScene()
 	// Need a RootActor for RemoveEmptyActors and to make AddChildActor agnostic to actor hiearchy level
 	TSharedPtr<IDatasmithActorElement> RootActor = FDatasmithSceneFactory::CreateActor(TEXT("RootActor"));
 	DatasmithScene->AddActor(RootActor);
-	ImportHierarchy(C4dDocument->GetFirstObject(), C4dDocument->GetFirstObject(), RootActor, melange::Matrix(), "", nullptr);
+	TArray<melange::TextureTag*> TextureTags;
+	ImportHierarchy(C4dDocument->GetFirstObject(), C4dDocument->GetFirstObject(), RootActor, melange::Matrix(), "", nullptr, TextureTags);
 
 	// Animations
 	LevelSequence = FDatasmithSceneFactory::CreateLevelSequence(DatasmithScene->GetName());
diff --git a/Engine/Plugins/Enterprise/DatasmithC4DImporter/Source/DatasmithC4DTranslator/Private/DatasmithC4DImporter.h b/Engine/Plugins/Enterprise/DatasmithC4DImporter/Source/DatasmithC4DTranslator/Private/DatasmithC4DImporter.h
--- a/Engine/Plugins/Enterprise/DatasmithC4DImporter/Source/DatasmithC4DTranslator/Private/DatasmithC4DImporter.h
+++ b/Engine/Plugins/Enterprise/DatasmithC4DImporter/Source/DatasmithC4DTranslator/Private/DatasmithC4DImporter.h
@@ -116,10 +116,10 @@ public:
 	TMap<int32, FString> GetCustomizedMaterialAssignment(const FString& DatasmithMeshName, const TArray<melange::TextureTag*>& TextureTags);
 
 	/** Imports a melange actor, which might involve parsing another small hierarchy of subnodes and deformers*/
-	TSharedPtr<IDatasmithActorElement> ImportObjectAndChildren(melange::BaseObject* ActorObject, melange::BaseObject* DataObject, TSharedPtr<IDatasmithActorElement> ParentActor, const melange::Matrix& WorldTransformMatrix, const FString& InstancePath, TArray<melange::BaseObject*>* InstanceObjects, const FString& DatasmithLabel);
+	TSharedPtr<IDatasmithActorElement> ImportObjectAndChildren(melange::BaseObject* ActorObject, melange::BaseObject* DataObject, TSharedPtr<IDatasmithActorElement> ParentActor, const melange::Matrix& WorldTransformMatrix, const FString& InstancePath, TArray<melange::BaseObject*>* InstanceObjects, const FString& DatasmithLabel, const TArray<melange::TextureTag*>& TextureTags);
 
 	/** Traverse the melange actor hierarchy importing all nodes */
-	void ImportHierarchy(melange::BaseObject* ActorObject, melange::BaseObject* DataObject, TSharedPtr<IDatasmithActorElement> ParentActor, const melange::Matrix& WorldTransformMatrix, const FString& InstancePath, TArray<melange::BaseObject*>* InstanceObjects);
+	void ImportHierarchy(melange::BaseObject* ActorObject, melange::BaseObject* DataObject, TSharedPtr<IDatasmithActorElement> ParentActor, const melange::Matrix& WorldTransformMatrix, const FString& InstancePath, TArray<melange::BaseObject*>* InstanceObjects, const TArray<melange::TextureTag*>& TextureTags);
 
 	/**
 	 * Adds Actor as a child of ParentActor using the corresponding WorldTransformMatrix. Object is the corresponding melange Object to Actor.
@@ -135,7 +135,7 @@ public:
 
 private:
 	/** Returns the TextureTags that should affect this object. May check parent objects, so relies on CachesOriginalObject */
-	TArray<melange::TextureTag*> GetActiveTextureTags(const melange::BaseObject* Object);
+	TArray<melange::TextureTag*> GetActiveTextureTags(const melange::BaseObject* Object, const TArray<melange::TextureTag*>& OrderedTextureTags);
 
 	/** Removes from Context->Scene all empty actors that have a single child */
 	void RemoveEmptyActors();
