C4D: Merge "Fix extra geometry being imported in some cases" from cl 11208473

diff --git a/Engine/Plugins/Enterprise/DatasmithC4DImporter/Source/DatasmithC4DTranslator/Private/DatasmithC4DImporter.cpp b/Engine/Plugins/Enterprise/DatasmithC4DImporter/Source/DatasmithC4DTranslator/Private/DatasmithC4DImporter.cpp
--- a/Engine/Plugins/Enterprise/DatasmithC4DImporter/Source/DatasmithC4DTranslator/Private/DatasmithC4DImporter.cpp
+++ b/Engine/Plugins/Enterprise/DatasmithC4DImporter/Source/DatasmithC4DTranslator/Private/DatasmithC4DImporter.cpp
@@ -516,29 +516,25 @@ namespace FC4DImporterImpl
 	// For now, we can't remove parents of animated nodes because animations are stored wrt local coordinate
 	// system. If we optimized an otherwise useless intermediate node, we'd need to bake its transform into all animations of
 	// child nodes, which I'm not sure is the ideal behavior as imported animation curves would look very different
-	bool KeepParentsOfAnimatedNodes(const TSharedPtr<IDatasmithActorElement>& Actor, TSet<FString>& NamesOfAnimatedActors)
+	bool KeepParentsOfAnimatedNodes(const TSharedPtr<IDatasmithActorElement>& Actor, TSet<FString>& NamesOfActorsToKeep)
 	{
-		bool bKeepThisNode = NamesOfAnimatedActors.Contains(Actor->GetName());
+		bool bKeepThisNode = NamesOfActorsToKeep.Contains(Actor->GetName());
 
 		for (int32 ChildIndex = 0; ChildIndex < Actor->GetChildrenCount(); ++ChildIndex)
 		{
-			bKeepThisNode |= KeepParentsOfAnimatedNodes(Actor->GetChild(ChildIndex), NamesOfAnimatedActors);
+			bKeepThisNode |= KeepParentsOfAnimatedNodes(Actor->GetChild(ChildIndex), NamesOfActorsToKeep);
 		}
 
 		if (bKeepThisNode)
 		{
-			NamesOfAnimatedActors.Add(Actor->GetName());
+			NamesOfActorsToKeep.Add(Actor->GetName());
 		}
 
 		return bKeepThisNode;
 	}
 
-	void RemoveEmptyActors(TSharedRef<IDatasmithScene>& DatasmithScene, const TSet<FString>& NamesOfCameraTargetActors, const TSet<FString>& NamesOfAnimatedActors)
+	void RemoveEmptyActors(TSharedRef<IDatasmithScene>& DatasmithScene, const TSet<FString>& NamesOfActorsToKeep)
 	{
-		TSet<FString> NamesOfActorsToKeep;
-		NamesOfActorsToKeep.Append(NamesOfCameraTargetActors);
-		NamesOfActorsToKeep.Append(NamesOfAnimatedActors);
-
 		for (int32 ActorIndex = 0; ActorIndex < DatasmithScene->GetActorsCount(); ++ActorIndex)
 		{
 			TSharedPtr<IDatasmithActorElement> Actor = DatasmithScene->GetActor(ActorIndex);
@@ -1143,7 +1139,7 @@ TSharedPtr<IDatasmithCameraActorElement> FDatasmithC4DImporter::ImportCamera(mel
 		}
 		CameraActor->SetLookAtActor(*(LookAtID.GetValue()));
 		CameraActor->SetLookAtAllowRoll(true);
-		NamesOfCameraTargetActors.Add(LookAtID.GetValue());
+		NamesOfActorsToKeep.Add(LookAtID.GetValue());
 	}
 
 	float CameraFocusDistanceInCM = MelangeGetFloat(InC4DCameraPtr, melange::CAMERAOBJECT_TARGETDISTANCE);
@@ -2138,7 +2134,7 @@ void FDatasmithC4DImporter::ImportAnimations(TSharedPtr<IDatasmithActorElement>
 	}
 
 	// Prevent actor from being optimized away
-	NamesOfAnimatedActors.Add(ActorElement->GetName());
+	NamesOfActorsToKeep.Add(ActorElement->GetName());
 
 	// Add a visibility track to simulate the particle spawning and despawning, if this is a particle actor.
 	// It seems like the particles have keys where they are visible: Before the first key the particles haven't
@@ -2494,6 +2490,25 @@ TSharedPtr<IDatasmithActorElement> FDatasmithC4DImporter::ImportObjectAndChildre
 
 				break;
 
+			case Oatomarray:
+			case Oconnector:
+				// Connector object will have as children the original objects, and its data cache will point at the polygon that results from the actual connect operation,
+				// so here we skip that hierarchy and just imoprt that polygon directly
+				ActorElement = ImportNullActor(ActorObject, DatasmithName.GetValue(), DatasmithLabel);
+
+				// This will be an empty actor, but we would like to keep it around because its the actor that receives the name of the connect object node itself,
+				// while its polygon seems to randomly receive the name of one of the original objects. Keeping the hierarchy like this makes it look exactly like what is
+				// shown in C4D if you make a connect object editable
+				NamesOfActorsToKeep.Add(ActorElement->GetName());
+				if (DataCache != nullptr && AddChildActor(ActorObject, ParentActor, NewWorldTransformMatrix, ActorElement))
+				{
+					ImportHierarchy(ActorCache, DataCache, ActorElement, NewWorldTransformMatrix, InstancePath, nullptr, TextureTags);
+					return ActorElement;
+				}
+
+				bSuccess = false;
+				break;
+
 			case Ofracture:
 			case ID_MOTIONFRACTUREVORONOI:
 			case Osymmetry:
@@ -3055,8 +3070,8 @@ bool FDatasmithC4DImporter::ProcessScene()
 	ImportActorHierarchyAnimations(RootActor);
 
 	// Processing
-	FC4DImporterImpl::KeepParentsOfAnimatedNodes(RootActor, NamesOfAnimatedActors);
-	FC4DImporterImpl::RemoveEmptyActors(DatasmithScene, NamesOfCameraTargetActors, NamesOfAnimatedActors);
+	FC4DImporterImpl::KeepParentsOfAnimatedNodes(RootActor, NamesOfActorsToKeep);
+	FC4DImporterImpl::RemoveEmptyActors(DatasmithScene, NamesOfActorsToKeep);
 	DatasmithScene->RemoveActor(RootActor, EDatasmithActorRemovalRule::KeepChildrenAndKeepRelativeTransform);
 
 #if WITH_EDITOR
diff --git a/Engine/Plugins/Enterprise/DatasmithC4DImporter/Source/DatasmithC4DTranslator/Private/DatasmithC4DImporter.h b/Engine/Plugins/Enterprise/DatasmithC4DImporter/Source/DatasmithC4DTranslator/Private/DatasmithC4DImporter.h
--- a/Engine/Plugins/Enterprise/DatasmithC4DImporter/Source/DatasmithC4DTranslator/Private/DatasmithC4DImporter.h
+++ b/Engine/Plugins/Enterprise/DatasmithC4DImporter/Source/DatasmithC4DTranslator/Private/DatasmithC4DImporter.h
@@ -174,8 +174,7 @@ private:
 	TSet<FString> NamesOfAllActors;
 
 	/** Names of IDatasmithActorElements that shouldn't be removed when optimizing the scene */
-	TSet<FString> NamesOfCameraTargetActors;
-	TSet<FString> NamesOfAnimatedActors;
+	TSet<FString> NamesOfActorsToKeep;
 
 	/** Where all actor animations are imported into when parsing the scene */
 	TSharedPtr<IDatasmithLevelSequenceElement> LevelSequence;
