C4D: Always set a set of materials on base mesh, even if we have multiple overrides

diff --git a/Engine/Plugins/Enterprise/DatasmithC4DImporter/Source/DatasmithC4DTranslator/Private/DatasmithC4DImporter.cpp b/Engine/Plugins/Enterprise/DatasmithC4DImporter/Source/DatasmithC4DTranslator/Private/DatasmithC4DImporter.cpp
--- a/Engine/Plugins/Enterprise/DatasmithC4DImporter/Source/DatasmithC4DTranslator/Private/DatasmithC4DImporter.cpp
+++ b/Engine/Plugins/Enterprise/DatasmithC4DImporter/Source/DatasmithC4DTranslator/Private/DatasmithC4DImporter.cpp
@@ -1670,6 +1670,23 @@ TSharedPtr<IDatasmithMeshActorElement> FDatasmithC4DImporter::ImportPolygon(mela
 	MeshActorElement->SetLabel(*DatasmithActorLabel);
 	MeshActorElement->SetStaticMeshPathName(ResultMeshElement->GetName());
 
+	// Check if we still need to assign materials to the base mesh
+	bool bMeshHasMaterialAssignments = false;
+	for ( int32 SlotIndex = 0; SlotIndex < ResultMeshElement->GetMaterialSlotCount(); ++SlotIndex )
+	{
+		TSharedPtr<IDatasmithMaterialIDElement> SlotIDElement = ResultMeshElement->GetMaterialSlotAt(SlotIndex);
+		if ( !SlotIDElement.IsValid() )
+		{
+			continue;
+		}
+
+		if ( ResultMeshElement->GetMaterial( SlotIDElement->GetId() ) != nullptr )
+		{
+			bMeshHasMaterialAssignments = true;
+			break;
+		}
+	}
+
 	// Add material overrides
 	TMap<int32, FString> SlotIndexToMaterialName = GetCustomizedMaterialAssignment(DatasmithMeshName, TextureTags);
 	for (const TPair<int32, FString>& Pair : SlotIndexToMaterialName)
@@ -1680,9 +1697,16 @@ TSharedPtr<IDatasmithMeshActorElement> FDatasmithC4DImporter::ImportPolygon(mela
 		TSharedPtr<IDatasmithMasterMaterialElement>* FoundMaterial = MaterialNameToMaterialElement.Find(MaterialName);
 		if (FoundMaterial && FoundMaterial->IsValid())
 		{
-			TSharedRef<IDatasmithMaterialIDElement> MaterialIDElement = FDatasmithSceneFactory::CreateMaterialId(*MaterialName);
-			MaterialIDElement->SetId(SlotIndex);
-			MeshActorElement->AddMaterialOverride(MaterialIDElement);
+			if ( bMeshHasMaterialAssignments && ResultMeshElement->GetMaterial( SlotIndex ) != MaterialName )
+			{
+				TSharedRef<IDatasmithMaterialIDElement> MaterialIDElement = FDatasmithSceneFactory::CreateMaterialId( *MaterialName );
+				MaterialIDElement->SetId( SlotIndex );
+				MeshActorElement->AddMaterialOverride( MaterialIDElement );
+			}
+			else if ( !bMeshHasMaterialAssignments )
+			{
+				ResultMeshElement->SetMaterial( *MaterialName, SlotIndex );
+			}
 		}
 	}
 
@@ -1691,7 +1715,7 @@ TSharedPtr<IDatasmithMeshActorElement> FDatasmithC4DImporter::ImportPolygon(mela
 
 namespace
 {
-	// Ignore the children of objects that are rendred using only their cache
+	// Ignore the children of objects that are rendered using only their cache
 	bool BrowseInstanceObjectChildren(melange::BaseObject* Object)
 	{
 		melange::Int32 ObjectType = Object->GetType();
@@ -3057,7 +3081,7 @@ bool FDatasmithC4DImporter::ProcessScene()
 	ImportedTextures.Empty();
 
 	// Actors
-	// Need a RootActor for RemoveEmptyActors and to make AddChildActor agnostic to actor hiearchy level
+	// Need a RootActor for RemoveEmptyActors and to make AddChildActor agnostic to actor hierarchy level
 	TSharedPtr<IDatasmithActorElement> RootActor = FDatasmithSceneFactory::CreateActor(TEXT("RootActor"));
 	DatasmithScene->AddActor(RootActor);
 	TArray<melange::TextureTag*> TextureTags;
