C4D: Fix import for linearly interpolated animations

diff --git a/Engine/Plugins/Enterprise/DatasmithC4DImporter/Source/DatasmithC4DTranslator/Private/DatasmithC4DImporter.cpp b/Engine/Plugins/Enterprise/DatasmithC4DImporter/Source/DatasmithC4DTranslator/Private/DatasmithC4DImporter.cpp
--- a/Engine/Plugins/Enterprise/DatasmithC4DImporter/Source/DatasmithC4DTranslator/Private/DatasmithC4DImporter.cpp
+++ b/Engine/Plugins/Enterprise/DatasmithC4DImporter/Source/DatasmithC4DTranslator/Private/DatasmithC4DImporter.cpp
@@ -2103,7 +2103,7 @@ void FDatasmithC4DImporter::ImportAnimations(TSharedPtr<IDatasmithActorElement>
 				TransformFrames,
 				InitialValues);
 
-			if (Interpolation != melange::CINTERPOLATION_LINEAR && KeyIndex < Curve->GetKeyCount() - 1)
+			if (KeyIndex < Curve->GetKeyCount() - 1)
 			{
 				//"Bake" the animation by generating a key for each frame between this Key and the next one
 				melange::CKey* NextKey = Curve->GetKey(KeyIndex+1);
