C4D: Fix meshes with materials forcing their materials onto other instances of that mesh that don't have materials assigned

diff --git a/Engine/Plugins/Enterprise/DatasmithC4DImporter/Source/DatasmithC4DTranslator/Private/DatasmithC4DImporter.cpp b/Engine/Plugins/Enterprise/DatasmithC4DImporter/Source/DatasmithC4DTranslator/Private/DatasmithC4DImporter.cpp
--- a/Engine/Plugins/Enterprise/DatasmithC4DImporter/Source/DatasmithC4DTranslator/Private/DatasmithC4DImporter.cpp
+++ b/Engine/Plugins/Enterprise/DatasmithC4DImporter/Source/DatasmithC4DTranslator/Private/DatasmithC4DImporter.cpp
@@ -1694,16 +1694,23 @@ TSharedPtr<IDatasmithMeshActorElement> FDatasmithC4DImporter::ImportPolygon(mela
 		int32 SlotIndex = Pair.Key;
 		const FString& MaterialName = Pair.Value;
 
+		// If we have an unassigned material, we *must* set it on the base mesh, as we can't create a material overrides to "clear" a material slot.
+		// Because of this we must always create and set material overrides for all found materials (below), as a PolygonObject imported later may
+		// cause the base mesh material to be reset to unassigned
+		if (MaterialName.IsEmpty())
+		{
+			ResultMeshElement->SetMaterial(TEXT(""), SlotIndex);
+			continue;
+		}
+
 		TSharedPtr<IDatasmithMasterMaterialElement>* FoundMaterial = MaterialNameToMaterialElement.Find(MaterialName);
 		if (FoundMaterial && FoundMaterial->IsValid())
 		{
-			if ( bMeshHasMaterialAssignments && ResultMeshElement->GetMaterial( SlotIndex ) != MaterialName )
-			{
-				TSharedRef<IDatasmithMaterialIDElement> MaterialIDElement = FDatasmithSceneFactory::CreateMaterialId( *MaterialName );
-				MaterialIDElement->SetId( SlotIndex );
-				MeshActorElement->AddMaterialOverride( MaterialIDElement );
-			}
-			else if ( !bMeshHasMaterialAssignments )
+			TSharedRef<IDatasmithMaterialIDElement> MaterialIDElement = FDatasmithSceneFactory::CreateMaterialId( *MaterialName );
+			MaterialIDElement->SetId( SlotIndex );
+			MeshActorElement->AddMaterialOverride( MaterialIDElement );
+
+			if ( !bMeshHasMaterialAssignments )
 			{
 				ResultMeshElement->SetMaterial( *MaterialName, SlotIndex );
 			}
