- Enterprise plugin changes for new Restricted locations (mostly removing the ThirdParty\NorForLIcensees type files to Restricted\NotForLicensees)

diff --git a/Engine/Plugins/Enterprise/AxFImporter/Source/AxFImporter/AxFImporter.Build.cs b/Engine/Plugins/Enterprise/AxFImporter/Source/AxFImporter/AxFImporter.Build.cs
--- a/Engine/Plugins/Enterprise/AxFImporter/Source/AxFImporter/AxFImporter.Build.cs
+++ b/Engine/Plugins/Enterprise/AxFImporter/Source/AxFImporter/AxFImporter.Build.cs
@@ -10,7 +10,7 @@ namespace UnrealBuildTool.Rules
     {
         private string ThirdPartyPath
         {
-            get { return Path.GetFullPath(Path.Combine(ModuleDirectory, "../ThirdParty/NotForLicensees/")); }
+            get { return Path.GetFullPath(Path.Combine(EngineDirectory, "Restricted/NotForLicensees/Plugins/Enterprise/AxFImporter/Source/ThirdParty")); }
         }
 
         public AxFImporter(ReadOnlyTargetRules Target) : base(Target)
@@ -57,7 +57,6 @@ namespace UnrealBuildTool.Rules
                 RuntimeDependencies.Add(ModulePath);
             }
 
-
             if (Directory.Exists(ThirdPartyPath))
             {
                 //third party libraries
diff --git a/Engine/Plugins/Enterprise/DatasmithC4DImporter/Source/DatasmithC4DTranslator/DatasmithC4DTranslator.Build.cs b/Engine/Plugins/Enterprise/DatasmithC4DImporter/Source/DatasmithC4DTranslator/DatasmithC4DTranslator.Build.cs
--- a/Engine/Plugins/Enterprise/DatasmithC4DImporter/Source/DatasmithC4DTranslator/DatasmithC4DTranslator.Build.cs
+++ b/Engine/Plugins/Enterprise/DatasmithC4DImporter/Source/DatasmithC4DTranslator/DatasmithC4DTranslator.Build.cs
@@ -47,7 +47,7 @@ namespace UnrealBuildTool.Rules
 			);
 
 			// Set up the C4D Melange SDK includes and libraries.
-			string MelangeSDKLocation = Path.Combine(PluginDirectory, "Source", "ThirdParty", "NotForLicensees", "Melange", "20.004_RBMelange20.0_259890");
+			string MelangeSDKLocation = Path.Combine(EngineDirectory, "Restricted/NotForLicensees/Plugins/Enterprise/DatasmithC4DTranslator/Source/ThirdParty/Melange/20.004_RBMelange20.0_259890");
 
 			// When C4D Melange SDK is not part of the developer's workspace, look for environment variable Melange_SDK.
 			if (!Directory.Exists(MelangeSDKLocation))
diff --git a/Engine/Plugins/Enterprise/DatasmithIFCImporter/Source/DatasmithIFCTranslator/DatasmithIFCTranslator.Build.cs b/Engine/Plugins/Enterprise/DatasmithIFCImporter/Source/DatasmithIFCTranslator/DatasmithIFCTranslator.Build.cs
--- a/Engine/Plugins/Enterprise/DatasmithIFCImporter/Source/DatasmithIFCTranslator/DatasmithIFCTranslator.Build.cs
+++ b/Engine/Plugins/Enterprise/DatasmithIFCImporter/Source/DatasmithIFCTranslator/DatasmithIFCTranslator.Build.cs
@@ -33,7 +33,7 @@ namespace UnrealBuildTool.Rules
 				PrivateDependencyModuleNames.Add("MessageLog");
 			}
 
-			string IfcEngineDir = Path.Combine(PluginDirectory, "Source", "ThirdParty", "NotForLicensees", "ifcengine");
+			string IfcEngineDir = Path.Combine(EngineDirectory, "Restricted/NotForLicensees/Plugins/Enterprise/DatasmithIFCTranslator/Source/ThirdParty/ifcengine");
 			if (Directory.Exists(IfcEngineDir))
 			{
 				PrivateDependencyModuleNames.AddRange(
diff --git a/Engine/Plugins/Enterprise/MDLImporter/Source/MDLImporter/MDLImporter.Build.cs b/Engine/Plugins/Enterprise/MDLImporter/Source/MDLImporter/MDLImporter.Build.cs
--- a/Engine/Plugins/Enterprise/MDLImporter/Source/MDLImporter/MDLImporter.Build.cs
+++ b/Engine/Plugins/Enterprise/MDLImporter/Source/MDLImporter/MDLImporter.Build.cs
@@ -11,7 +11,7 @@ namespace UnrealBuildTool.Rules
     {
         private string ThirdPartyPath
         {
-            get { return Path.GetFullPath(Path.Combine(ModuleDirectory, "../ThirdParty/NotForLicensees/")); }
+            get { return Path.GetFullPath(Path.Combine(EngineDirectory, "Restricted/NotForLicensees/Plugins/Enterprise/MDLImporter/Source/ThirdParty/")); }
         }
 
         public MDLImporter(ReadOnlyTargetRules Target) : base(Target)
diff --git a/Engine/Plugins/Enterprise/StaticMeshEditorExtension/Source/MeshProcessingLibrary/MeshProcessingLibrary.Build.cs b/Engine/Plugins/Enterprise/StaticMeshEditorExtension/Source/MeshProcessingLibrary/MeshProcessingLibrary.Build.cs
--- a/Engine/Plugins/Enterprise/StaticMeshEditorExtension/Source/MeshProcessingLibrary/MeshProcessingLibrary.Build.cs
+++ b/Engine/Plugins/Enterprise/StaticMeshEditorExtension/Source/MeshProcessingLibrary/MeshProcessingLibrary.Build.cs
@@ -42,7 +42,7 @@ public class MeshProcessingLibrary : ModuleRules
 		}
 
 		// Setup MeshSimplifier
-		string MeshSimplifierPath = Path.Combine(ModuleDirectory, "ThirdParty/NotForLicensees/MeshSimplifier");
+		string MeshSimplifierPath = Path.Combine(EngineDirectory, "Restricted/NotForLicensees/Plugins/Enterprise/MeshProcessingLibrary/Source/ThirdParty/MeshSimplifier");
 		bool bWithMeshSimplifier = Directory.Exists(MeshSimplifierPath) && Target.Platform == UnrealTargetPlatform.Win64;
 		PublicDefinitions.Add("WITH_MESH_SIMPLIFIER=" + (bWithMeshSimplifier ? '1' : '0'));
 
