C4D: Limit number of imported UV channels to 7

diff --git a/Engine/Plugins/Enterprise/DatasmithC4DImporter/Source/DatasmithC4DTranslator/Private/DatasmithC4DImporter.cpp b/Engine/Plugins/Enterprise/DatasmithC4DImporter/Source/DatasmithC4DTranslator/Private/DatasmithC4DImporter.cpp
--- a/Engine/Plugins/Enterprise/DatasmithC4DImporter/Source/DatasmithC4DTranslator/Private/DatasmithC4DImporter.cpp
+++ b/Engine/Plugins/Enterprise/DatasmithC4DImporter/Source/DatasmithC4DTranslator/Private/DatasmithC4DImporter.cpp
@@ -2728,6 +2728,12 @@ TSharedPtr<IDatasmithMeshElement> FDatasmithC4DImporter::ImportMesh(melange::Pol
 		}
 	}
 
+	if (UVWTagsData.Num() > MAX_STATIC_TEXCOORDS-1)
+	{
+		UE_LOG(LogDatasmithC4DImport, Error, TEXT("Mesh '%s' has %d UV channels! Only the first %d will be used"), *DatasmithLabel, MAX_STATIC_TEXCOORDS-1, UVWTagsData.Num());
+		UVWTagsData.SetNum(MAX_STATIC_TEXCOORDS-1);
+	}
+
 	int32 NumSlots = SelectionTags.Num();
 
 	// Create MeshDescription
