C4D: Updated path the Melange

diff --git a/Engine/Plugins/Enterprise/DatasmithC4DImporter/Source/DatasmithC4DTranslator/DatasmithC4DTranslator.Build.cs b/Engine/Plugins/Enterprise/DatasmithC4DImporter/Source/DatasmithC4DTranslator/DatasmithC4DTranslator.Build.cs
--- a/Engine/Plugins/Enterprise/DatasmithC4DImporter/Source/DatasmithC4DTranslator/DatasmithC4DTranslator.Build.cs
+++ b/Engine/Plugins/Enterprise/DatasmithC4DImporter/Source/DatasmithC4DTranslator/DatasmithC4DTranslator.Build.cs
@@ -47,7 +47,7 @@ namespace UnrealBuildTool.Rules
 			);
 
 			// Set up the C4D Melange SDK includes and libraries.
-			string MelangeSDKLocation = Path.Combine(EngineDirectory, "Restricted/NotForLicensees/Plugins/Enterprise/DatasmithC4DTranslator/Source/ThirdParty/Melange/20.004_RBMelange20.0_259890");
+			string MelangeSDKLocation = Path.Combine(EngineDirectory, "Restricted/NotForLicensees/Source/ThirdParty/Enterprise/Melange/20.004_RBMelange20.0_259890");
 
 			// When C4D Melange SDK is not part of the developer's workspace, look for environment variable Melange_SDK.
 			if (!Directory.Exists(MelangeSDKLocation))
