Fix compiliation issue when using c4d in combination with openEXR:
- Add _HAS_EXCEPTIONS in the backuped macros before including c4d.h

Note: c4d changes _HAS_EXCEPTIONS (see c4d_misc\general.h) which breaks subsequent usage of openEXR

diff --git a/Engine/Plugins/Enterprise/DatasmithC4DImporter/Source/DatasmithC4DTranslator/Private/DatasmithC4DMelangeSDKEnterGuard.h b/Engine/Plugins/Enterprise/DatasmithC4DImporter/Source/DatasmithC4DTranslator/Private/DatasmithC4DMelangeSDKEnterGuard.h
--- a/Engine/Plugins/Enterprise/DatasmithC4DImporter/Source/DatasmithC4DTranslator/Private/DatasmithC4DMelangeSDKEnterGuard.h
+++ b/Engine/Plugins/Enterprise/DatasmithC4DImporter/Source/DatasmithC4DTranslator/Private/DatasmithC4DMelangeSDKEnterGuard.h
@@ -7,6 +7,7 @@
 // Back up some preprocessor identifiers.
 __pragma(push_macro("PI"))
 __pragma(push_macro("BYTE_MAX"))
+__pragma(push_macro("_HAS_EXCEPTIONS"))
 
 #undef PI
 
diff --git a/Engine/Plugins/Enterprise/DatasmithC4DImporter/Source/DatasmithC4DTranslator/Private/DatasmithC4DMelangeSDKLeaveGuard.h b/Engine/Plugins/Enterprise/DatasmithC4DImporter/Source/DatasmithC4DTranslator/Private/DatasmithC4DMelangeSDKLeaveGuard.h
--- a/Engine/Plugins/Enterprise/DatasmithC4DImporter/Source/DatasmithC4DTranslator/Private/DatasmithC4DMelangeSDKLeaveGuard.h
+++ b/Engine/Plugins/Enterprise/DatasmithC4DImporter/Source/DatasmithC4DTranslator/Private/DatasmithC4DMelangeSDKLeaveGuard.h
@@ -7,6 +7,7 @@ __pragma(warning(pop))
 // Restore some preprocessor identifiers.
 __pragma(pop_macro("PI"))
 __pragma(pop_macro("BYTE_MAX"))
+__pragma(pop_macro("_HAS_EXCEPTIONS"))
 
 // Leave Datasmith platform include guard.
 #include "Windows/HideWindowsPlatformTypes.h"
