Fix include path casing

diff --git a/Engine/Plugins/Enterprise/DatasmithC4DImporter/Source/DatasmithC4DTranslator/Private/DatasmithC4DImporter.cpp b/Engine/Plugins/Enterprise/DatasmithC4DImporter/Source/DatasmithC4DTranslator/Private/DatasmithC4DImporter.cpp
--- a/Engine/Plugins/Enterprise/DatasmithC4DImporter/Source/DatasmithC4DTranslator/Private/DatasmithC4DImporter.cpp
+++ b/Engine/Plugins/Enterprise/DatasmithC4DImporter/Source/DatasmithC4DTranslator/Private/DatasmithC4DImporter.cpp
@@ -43,7 +43,7 @@
 #include "StaticMeshOperations.h"
 #include "Widgets/Notifications/SNotificationList.h"
 
-#include "OpenEXR/ImathMatrixAlgo.h"
+#include "openexr/ImathMatrixAlgo.h"
 
 DECLARE_CYCLE_STAT(TEXT("C4DImporter - Load File"), STAT_C4DImporter_LoadFile, STATGROUP_C4DImporter);
 
diff --git a/Engine/Source/Developer/Datasmith/DatasmithFacade/Private/DatasmithFacadeActor.cpp b/Engine/Source/Developer/Datasmith/DatasmithFacade/Private/DatasmithFacadeActor.cpp
--- a/Engine/Source/Developer/Datasmith/DatasmithFacade/Private/DatasmithFacadeActor.cpp
+++ b/Engine/Source/Developer/Datasmith/DatasmithFacade/Private/DatasmithFacadeActor.cpp
@@ -8,7 +8,7 @@
 #include "DatasmithFacadeScene.h"
 
 // OpenEXR third party library.
-#include "OpenEXR/ImathMatrixAlgo.h"
+#include "openexr/ImathMatrixAlgo.h"
 
 
 FDatasmithFacadeActor::FDatasmithFacadeActor(
diff --git a/Engine/Source/Programs/Enterprise/Datasmith/DatasmithSketchUpExporter/Private/DatasmithSketchUpComponent.cpp b/Engine/Source/Programs/Enterprise/Datasmith/DatasmithSketchUpExporter/Private/DatasmithSketchUpComponent.cpp
--- a/Engine/Source/Programs/Enterprise/Datasmith/DatasmithSketchUpExporter/Private/DatasmithSketchUpComponent.cpp
+++ b/Engine/Source/Programs/Enterprise/Datasmith/DatasmithSketchUpExporter/Private/DatasmithSketchUpComponent.cpp
@@ -28,7 +28,7 @@
 #include "DatasmithUtils.h"
 
 // OpenEXR third party library.
-#include "OpenEXR/ImathMatrixAlgo.h"
+#include "openexr/ImathMatrixAlgo.h"
 
 
 #ifdef SKP_SDK_2017
