Replace occurrences of UE4 with UE for VRED, DeltaGen, C4D and USD importers

diff --git a/Engine/Plugins/Enterprise/DatasmithC4DImporter/Source/DatasmithC4DTranslator/Private/DatasmithC4DImportOptions.h b/Engine/Plugins/Enterprise/DatasmithC4DImporter/Source/DatasmithC4DTranslator/Private/DatasmithC4DImportOptions.h
--- a/Engine/Plugins/Enterprise/DatasmithC4DImporter/Source/DatasmithC4DTranslator/Private/DatasmithC4DImportOptions.h
+++ b/Engine/Plugins/Enterprise/DatasmithC4DImporter/Source/DatasmithC4DTranslator/Private/DatasmithC4DImportOptions.h
@@ -47,7 +47,7 @@ public:
 
 	/**
 	* Scale all position and vertices by this value.
-	* It is needed when scene has too precise meshes because the conversion from double (C4D) to float (UE4)
+	* It is needed when scene has too precise meshes because the conversion from double (C4D) to float (UnrealEditor)
 	* can "degenerate" the faces
 	*/
 	UPROPERTY(config, EditAnywhere, Category = DebugProperty, meta = (DisplayName = "Scale the Entire Scene"))
diff --git a/Engine/Plugins/Enterprise/DatasmithC4DImporter/Source/DatasmithC4DTranslator/Private/DatasmithC4DImporter.cpp b/Engine/Plugins/Enterprise/DatasmithC4DImporter/Source/DatasmithC4DTranslator/Private/DatasmithC4DImporter.cpp
--- a/Engine/Plugins/Enterprise/DatasmithC4DImporter/Source/DatasmithC4DTranslator/Private/DatasmithC4DImporter.cpp
+++ b/Engine/Plugins/Enterprise/DatasmithC4DImporter/Source/DatasmithC4DTranslator/Private/DatasmithC4DImporter.cpp
@@ -53,8 +53,8 @@ DEFINE_LOG_CATEGORY(LogDatasmithC4DImport);
 
 // What we multiply the C4D light brightness values with when the lights are not
 // using photometric units. Those are chosen so that 100% brightness C4D point lights matches the
-// default value of 8 candelas of UE4 point lights, and 100% brightness C4D infinite lights matches
-// the default 10 lux of UE4 directional lights
+// default value of 8 candelas of UnrealEditor point lights, and 100% brightness C4D infinite lights matches
+// the default 10 lux of UnrealEditor directional lights
 #define UnitlessGlobalLightIntensity 10.0
 #define UnitlessIESandPointLightIntensity 8000
 
@@ -276,11 +276,11 @@ TSharedRef<FCraneCameraAttributes> ExtractCraneCameraAttributes(melange::BaseTag
 // coordinate system
 FTransform CalculateCraneCameraTransform(const FCraneCameraAttributes& Params)
 {
-	// We will first construct a transformation in the UE4 coordinate system, as that is
+	// We will first construct a transformation in the UnrealEditor coordinate system, as that is
 	// easier to visualize and test
 
 	// Local rotation of 90deg around the Y axis in Melange.
-	// Will compensate the difference in convention between UE4 (camera shoots out the +X) and
+	// Will compensate the difference in convention between UnrealEditor (camera shoots out the +X) and
 	// C4D (camera shoots out the +Z)
 	FTransform Conv = FTransform(FRotator(0, -90, 0), FVector(0, 0, 0));
 
@@ -318,17 +318,17 @@ FTransform CalculateCraneCameraTransform(const FCraneCameraAttributes& Params)
 			  Arm;
 	}
 
-	FTransform FinalTransUE4 = Conv * Cam * Head * Arm * Base;
-	FVector TranslationUE4 = FinalTransUE4.GetTranslation();
-	FVector EulerUE4 = FinalTransUE4.GetRotation().Euler();
+	FTransform FinalTransUE = Conv * Cam * Head * Arm * Base;
+	FVector TranslationUE = FinalTransUE.GetTranslation();
+	FVector EulerUE = FinalTransUE.GetRotation().Euler();
 
-	// Convert FinalTransUE4 into the melange coordinate system, so that this can be treated
+	// Convert FinalTransUnrealEditor into the melange coordinate system, so that this can be treated
 	// like the other types of animations in ImportAnimations.
 	// More specifically, convert them so that that ConvertDirectionLeftHandedYup and
-	// the conversion for Ocamera rotations gets them back into UE4's coordinate system
+	// the conversion for Ocamera rotations gets them back into UnrealEditor's coordinate system
 	// Note: Remember that FRotator's constructor is Pitch, Yaw and Roll (i.e. Y, Z, X)
-	return FTransform(FRotator(EulerUE4.Y, EulerUE4.X, -EulerUE4.Z-90),
-					  FVector(TranslationUE4.X, TranslationUE4.Z, -TranslationUE4.Y));
+	return FTransform(FRotator(EulerUE.Y, EulerUE.X, -EulerUE.Z-90),
+					  FVector(TranslationUE.X, TranslationUE.Z, -TranslationUE.Y));
 }
 
 void FDatasmithC4DImporter::ImportSpline(melange::SplineObject* SplineActor)
@@ -700,7 +700,7 @@ bool FDatasmithC4DImporter::AddChildActor(melange::BaseObject* Object, TSharedPt
 	if (Object->GetType() == Ocamera || Object->GetType() == Olight)
 	{
 		// Compensates the fact that in C4D cameras/lights shoot out towards +Z, while in
-		// UE4 they shoot towards +X
+		// UnrealEditor they shoot towards +X
 		melange::Matrix CameraRotation(
 			melange::Vector(0.0, 0.0, 0.0),
 			melange::Vector(0.0, 0.0, 1.0),
@@ -2298,13 +2298,14 @@ void FDatasmithC4DImporter::ImportAnimations(TSharedPtr<IDatasmithActorElement>
 				}
 
 				// TransformValue represents, in radians, the rotations around the C4D axes
-				// XRot, YRot, ZRot are rotations around UE4 axes, in the UE4 CS, with the sign given by Quaternion rotations (NOT Rotators)
+				// XRot, YRot, ZRot are rotations around UnrealEditor axes, in the UnrealEditor CS,
+				// with the sign given by Quaternion rotations (NOT Rotators)
 				FQuat XRot = FQuat(FVector(1, 0, 0), -TransformValueCopy.X);
 				FQuat YRot = FQuat(FVector(0, 1, 0),  TransformValueCopy.Z);
 				FQuat ZRot = FQuat(FVector(0, 0, 1), -TransformValueCopy.Y);
 
-				// Swap YRot and ZRot in the composition order, as an XYZ order in the C4D CS really means a XZY order in the UE4 CS
-				// This effectively converts the rotation order from the C4D CS to the UE4 CS, the sign of the rotations being handled when
+				// Swap YRot and ZRot in the composition order, as an XYZ order in the C4D CS really means a XZY order in the UnrealEditor CS
+				// This effectively converts the rotation order from the C4D CS to the UnrealEditor CS, the sign of the rotations being handled when
 				// creating the FQuats
 				Swap(YRot, ZRot);
 
@@ -2333,7 +2334,7 @@ void FDatasmithC4DImporter::ImportAnimations(TSharedPtr<IDatasmithActorElement>
 					break;
 				}
 
-				// In C4D cameras and lights shoot towards +Z, but in UE4 they shoot towards +X, so fix that with a yaw
+				// In C4D cameras and lights shoot towards +Z, but in UnrealEditor they shoot towards +X, so fix that with a yaw
 				if (Object->GetType() == Olight || Object->GetType() == Ocamera)
 				{
 					FinalQuat = FinalQuat * FQuat(FVector(0, 0, 1), FMath::DegreesToRadians(-90.0f));
diff --git a/Engine/Plugins/Enterprise/DatasmithC4DImporter/Source/DatasmithC4DTranslator/Private/DatasmithC4DUtils.h b/Engine/Plugins/Enterprise/DatasmithC4DImporter/Source/DatasmithC4DTranslator/Private/DatasmithC4DUtils.h
--- a/Engine/Plugins/Enterprise/DatasmithC4DImporter/Source/DatasmithC4DTranslator/Private/DatasmithC4DUtils.h
+++ b/Engine/Plugins/Enterprise/DatasmithC4DImporter/Source/DatasmithC4DTranslator/Private/DatasmithC4DUtils.h
@@ -72,14 +72,14 @@ melange::BaseList2D* MelangeGetLink(melange::BaseList2D* Object, melange::Int32
 float MelangeGetFloat(melange::BaseList2D* Object, melange::Int32 Parameter);
 
 /**
- * Converts a vector representing a position from the melange coordinate system to the UE4 coordinate system
+ * Converts a vector representing a position from the melange coordinate system to the UnrealEditor coordinate system
  */
 FVector ConvertMelangePosition(const melange::Vector32& MelangePosition, float WorldUnitScale = 1.0f);
 FVector ConvertMelangePosition(const melange::Vector64& MelangePosition, float WorldUnitScale = 1.0f);
 FVector ConvertMelangePosition(const FVector& MelangePosition, float WorldUnitScale = 1.0f);
 
 /**
- * Converts a vector representing a direction from the melange coordinate system to the UE4 coordinate system
+ * Converts a vector representing a direction from the melange coordinate system to the UnrealEditor coordinate system
  */
 FVector ConvertMelangeDirection(const melange::Vector32& MelangePosition);
 FVector ConvertMelangeDirection(const melange::Vector64& MelangePosition);
diff --git a/Engine/Plugins/Enterprise/DatasmithFBXImporter/Resources/VREDPlugin/DatasmithExporter/anim_exporter.py b/Engine/Plugins/Enterprise/DatasmithFBXImporter/Resources/VREDPlugin/DatasmithExporter/anim_exporter.py
--- a/Engine/Plugins/Enterprise/DatasmithFBXImporter/Resources/VREDPlugin/DatasmithExporter/anim_exporter.py
+++ b/Engine/Plugins/Enterprise/DatasmithFBXImporter/Resources/VREDPlugin/DatasmithExporter/anim_exporter.py
@@ -150,7 +150,7 @@ class DSAnimationLibrary:
 
         # AnimClips can never have repeated names, and even in nested clips you can't have
         # clip references. So every single clip should have a unique name, as it will become
-        # an asset in UE4
+        # an asset in UnrealEditor
         clipsToSanitize = []
         for clip in self.DSClips:
             sanitizeRecursive(clip, clipsToSanitize)
diff --git a/Engine/Plugins/Enterprise/DatasmithFBXImporter/Resources/VREDPlugin/DatasmithExporter/exporter.py b/Engine/Plugins/Enterprise/DatasmithFBXImporter/Resources/VREDPlugin/DatasmithExporter/exporter.py
--- a/Engine/Plugins/Enterprise/DatasmithFBXImporter/Resources/VREDPlugin/DatasmithExporter/exporter.py
+++ b/Engine/Plugins/Enterprise/DatasmithFBXImporter/Resources/VREDPlugin/DatasmithExporter/exporter.py
@@ -275,9 +275,9 @@ class DatasmithFBXExporter():
 
     def processTransform(self, trans, scale):
         """
-        Flip and transpose the transform so it can be used directly in UE4. Also, scales
+        Flip and transpose the transform so it can be used directly in UnrealEditor. Also, scales
         the translation by the unit scale, which is something that's easier to do here than
-        in the UE4 importer.
+        in the UnrealEditor importer.
 
         Args:
             trans (list of 16 floats): Describes the transform to process
@@ -406,12 +406,12 @@ class DatasmithFBXExporter():
                     xmlNode.set('state', newStateName)
 
     def fixupXmlContents(self, root):
-        """ Works around a bug in UE4 by ensuring that element contents that spans multiple lines end in a newline
+        """ Works around a bug in UnrealEditor by ensuring that element contents that spans multiple lines end in a newline
 
         Without this pass, an xml file like this:
         <A>
         e"</A>
-        Will not be able to be parsed by UE4's XmlParser
+        Will not be able to be parsed by UnrealEditor's XmlParser
 
         Args:
             root (xml.etree.Element): Root of the xml subtree that will be fixed        
diff --git a/Engine/Plugins/Enterprise/DatasmithFBXImporter/Resources/VREDPlugin/DatasmithExporter/mat.py b/Engine/Plugins/Enterprise/DatasmithFBXImporter/Resources/VREDPlugin/DatasmithExporter/mat.py
--- a/Engine/Plugins/Enterprise/DatasmithFBXImporter/Resources/VREDPlugin/DatasmithExporter/mat.py
+++ b/Engine/Plugins/Enterprise/DatasmithFBXImporter/Resources/VREDPlugin/DatasmithExporter/mat.py
@@ -249,7 +249,7 @@ def unProcessTexture(tex):
         # In the end, we have no idea what has been done to this image when we have access to it.
         # 'nothing' seems to be the majority of cases here, so we'll ignore them for now.
         # Besides, we have to consider that this is a regular normal map on import, so the user
-        # can quickly switch to another source image on UE4 side and things will just work
+        # can quickly switch to another source image on UnrealEditor side and things will just work
         return
 
     elif tex.texType == 'displacementComponent':
diff --git a/Engine/Plugins/Enterprise/DatasmithFBXImporter/Source/DatasmithDeltaGenTranslator/Private/DatasmithDeltaGenImporter.cpp b/Engine/Plugins/Enterprise/DatasmithFBXImporter/Source/DatasmithDeltaGenTranslator/Private/DatasmithDeltaGenImporter.cpp
--- a/Engine/Plugins/Enterprise/DatasmithFBXImporter/Source/DatasmithDeltaGenTranslator/Private/DatasmithDeltaGenImporter.cpp
+++ b/Engine/Plugins/Enterprise/DatasmithFBXImporter/Source/DatasmithDeltaGenTranslator/Private/DatasmithDeltaGenImporter.cpp
@@ -844,7 +844,7 @@ namespace DeltaGenImporterImpl
 			}
 			else if (DSType == EDatasmithTransformType::Translation)
 			{
-				// Deltagen is right-handed Z up, UE4 is left-handed Z up. We try keeping the same X, so here we
+				// Deltagen is right-handed Z up, UnrealEditor is left-handed Z up. We try keeping the same X, so here we
 				// just flip the Y coordinate to convert between them.
 				// Note: Geometry, transforms and VRED animations get converted when parsing the FBX file in DatasmithFBXFileImporter, so
 				// you won't find analogue for this in VRED importer, even though the conversion is the same
diff --git a/Engine/Plugins/Enterprise/DatasmithFBXImporter/Source/DatasmithFBXTranslator/Private/DatasmithFBXFileImporter.cpp b/Engine/Plugins/Enterprise/DatasmithFBXImporter/Source/DatasmithFBXTranslator/Private/DatasmithFBXFileImporter.cpp
--- a/Engine/Plugins/Enterprise/DatasmithFBXImporter/Source/DatasmithFBXTranslator/Private/DatasmithFBXFileImporter.cpp
+++ b/Engine/Plugins/Enterprise/DatasmithFBXImporter/Source/DatasmithFBXTranslator/Private/DatasmithFBXFileImporter.cpp
@@ -247,7 +247,7 @@ void FDatasmithFBXFileImporter::TraverseHierarchyNodeRecursively(FbxNode* Parent
 
 		FVector RotEuler = Rotation.Euler();
 
-		// Avoid singularity around 90 degree pitch, as UE4 doesn't seem to support it very well
+		// Avoid singularity around 90 degree pitch, as UnrealEditor doesn't seem to support it very well
 		// See UE-75467 and UE-83049
 		if (FMath::IsNearlyEqual(abs(RotEuler.Y), 90.0f))
 		{
diff --git a/Engine/Plugins/Enterprise/DatasmithFBXImporter/Source/DatasmithVREDTranslator/Private/DatasmithVREDClipProcessor.cpp b/Engine/Plugins/Enterprise/DatasmithFBXImporter/Source/DatasmithVREDTranslator/Private/DatasmithVREDClipProcessor.cpp
--- a/Engine/Plugins/Enterprise/DatasmithFBXImporter/Source/DatasmithVREDTranslator/Private/DatasmithVREDClipProcessor.cpp
+++ b/Engine/Plugins/Enterprise/DatasmithFBXImporter/Source/DatasmithVREDTranslator/Private/DatasmithVREDClipProcessor.cpp
@@ -97,7 +97,7 @@ void FDatasmithVREDClipProcessor::BuildMaps()
 // starting at t=2s, clip B's contents will, themselves, all start at t=2s. This
 // is not how the sequencer works (which would play clip B starts at t=0),
 // so here we remove this "nested clip delay". It is also a lot handier because
-// clips can't be reused in VRED, but subsequences can be reused in UE4
+// clips can't be reused in VRED, but subsequences can be reused in UnrealEditor
 void FDatasmithVREDClipProcessor::RecursivelyRemoveNestedClipDelay(FDatasmithFBXSceneAnimClip* Clip, float Delay)
 {
     for (FDatasmithFBXSceneAnimUsage& Usage : Clip->AnimUsages)
diff --git a/Engine/Plugins/Enterprise/DatasmithFBXImporter/Source/DatasmithVREDTranslator/Private/DatasmithVREDClipProcessor.h b/Engine/Plugins/Enterprise/DatasmithFBXImporter/Source/DatasmithVREDTranslator/Private/DatasmithVREDClipProcessor.h
--- a/Engine/Plugins/Enterprise/DatasmithFBXImporter/Source/DatasmithVREDTranslator/Private/DatasmithVREDClipProcessor.h
+++ b/Engine/Plugins/Enterprise/DatasmithFBXImporter/Source/DatasmithVREDTranslator/Private/DatasmithVREDClipProcessor.h
@@ -30,7 +30,7 @@ private:
 	// starting at t=2s, clip B's contents will, themselves, all start at t=2s. This
 	// is not how the sequencer works (which would play clip B starts at t=0),
 	// so here we remove this "nested clip delay". It is also a lot handier because
-	// clips can't be reused in VRED, but subsequences can be reused in UE4
+	// clips can't be reused in VRED, but subsequences can be reused in UnrealEditor
 	void RecursivelyRemoveNestedClipDelay (FDatasmithFBXSceneAnimClip* Clip, float Delay);
 	void RecursivelyPropagateFlippedToBlockUsages(FDatasmithFBXSceneAnimClip* Clip, bool bFlipped);
 
diff --git a/Engine/Plugins/Enterprise/DatasmithFBXImporter/Source/DatasmithVREDTranslator/Private/DatasmithVREDImporter.cpp b/Engine/Plugins/Enterprise/DatasmithFBXImporter/Source/DatasmithVREDTranslator/Private/DatasmithVREDImporter.cpp
--- a/Engine/Plugins/Enterprise/DatasmithFBXImporter/Source/DatasmithVREDTranslator/Private/DatasmithVREDImporter.cpp
+++ b/Engine/Plugins/Enterprise/DatasmithFBXImporter/Source/DatasmithVREDTranslator/Private/DatasmithVREDImporter.cpp
@@ -1031,7 +1031,7 @@ void AddTextureMappingProperties(IDatasmithMasterMaterialElement* Element, const
 	AddBoolProperty(Element, TexHandle + TEXT("ProjectionIsPlanar"), MappingType == ETextureMapType::Planar);
 
 	// Send all parameters every time, as a material might have properties for multiple projection types
-	// and the user can decide to switch them while in UE4 already
+	// and the user can decide to switch them while in UnrealEditor already
 	AddBoolProperty(Element, TexHandle + TEXT("LimitedProjectionSize"), Tex.Scale.Z != FLT_MAX);
 	AddColorProperty(Element, TexHandle + TEXT("ProjectionCenter"), Tex.Translation);
 	AddColorProperty(Element, TexHandle + TEXT("ProjectionOrientation"), Tex.Rotation);
@@ -1158,7 +1158,7 @@ TSharedPtr<IDatasmithBaseMaterialElement> FDatasmithVREDImporter::ConvertMateria
 	// Base colors used for Chrome and BrushedMetal materials
 	using MetalColorEntry = TPairInitializer<const int32&, const FVector4&>;
 	const static TMap<int32, FVector4> MetalColors
-	({															       // From UE4 docs    // Experimental
+	({															// From UnrealEditor docs  // Experimental
 		MetalColorEntry( 0, FVector4(0.950f, 0.950f, 0.950f, 1.000f)),                     // Highly reflective
 		MetalColorEntry( 1, FVector4(0.913f, 0.921f, 0.925f, 1.000f)), // Aluminium
 		MetalColorEntry( 2, FVector4(0.300f, 0.300f, 0.280f, 1.000f)),                     // Amorphous Carbon
@@ -1197,7 +1197,7 @@ TSharedPtr<IDatasmithBaseMaterialElement> FDatasmithVREDImporter::ConvertMateria
 		}
 
 		// We do this so that all glass materials are transparent by default, while also allowing the
-		// UE4 user to move the opacity back to 1.0f and have it not be transparent. This also helps to match
+		// UnrealEditor user to move the opacity back to 1.0f and have it not be transparent. This also helps to match
 		// VRED's multiply blend mode behaviour for translucent materials
 		if (float* OldOpacity = Material->ScalarParams.Find(TEXT("Opacity")))
 		{
@@ -1247,7 +1247,7 @@ TSharedPtr<IDatasmithBaseMaterialElement> FDatasmithVREDImporter::ConvertMateria
 		}
 
 		// For raytracing, VRED allows different roughness directions. We don't support
-		// that in UE4, so let's just combine the roughness values like VRED does in its
+		// that in UnrealEditor, so let's just combine the roughness values like VRED does in its
 		// viewport so that it's easier to manipulate in the material instance
 		const float* RoughnessU = Material->ScalarParams.Find(TEXT("RoughnessU"));
 		const float* RoughnessV = Material->ScalarParams.Find(TEXT("RoughnessV"));
@@ -1456,7 +1456,7 @@ void PopulateTransformAnimation(IDatasmithTransformAnimationElement& TransformAn
 
 	// We go to EndFrame.Value+1 here so that if its a 2 second animation at 30fps, frame 60 belongs
 	// to the actual animation, as opposed to being range [0, 59]. I don't think this is the standard
-	// for UE4, but it is exactly how it works in VRED and it also guarantees that the animation will
+	// for UnrealEditor, but it is exactly how it works in VRED and it also guarantees that the animation will
 	// actually complete within its range, which is necessary in order to play it as a subsequence
 	// to its completion
 	for (int32 Frame = StartFrame.Value; Frame <= EndFrame.Value + 1; ++Frame)
@@ -1522,7 +1522,7 @@ void PopulateVisibilityAnimation(IDatasmithVisibilityAnimationElement& Visibilit
 
 	// We go to EndFrame.Value+1 here so that if its a 2 second animation at 30fps, frame 60 belongs
 	// to the actual animation, as opposed to being frame [0, 59]. I don't think this is the standard
-	// for UE4, but it is exactly how it works in VRED and it also guarantees that the animation will
+	// for UnrealEditor, but it is exactly how it works in VRED and it also guarantees that the animation will
 	// actually complete within its range, which is necessary in order to play it as a subsequence
 	// to its completion
 	for (int32 Frame = StartFrame.Value; Frame <= EndFrame.Value + 1; ++Frame)
diff --git a/Engine/Plugins/Importers/USDImporter/Source/USDUtilities/Private/USDSkeletalDataConversion.cpp b/Engine/Plugins/Importers/USDImporter/Source/USDUtilities/Private/USDSkeletalDataConversion.cpp
--- a/Engine/Plugins/Importers/USDImporter/Source/USDUtilities/Private/USDSkeletalDataConversion.cpp
+++ b/Engine/Plugins/Importers/USDImporter/Source/USDUtilities/Private/USDSkeletalDataConversion.cpp
@@ -243,7 +243,7 @@ namespace SkelDataConversionImpl
 		return true;
 	}
 
-	/** Converts the given offsets into UE4 space and fills in an FUsdBlendShape object with all the data that will become a morph target */
+	/** Converts the given offsets into UnrealEditor space and fills in an FUsdBlendShape object with all the data that will become a morph target */
 	bool CreateUsdBlendShape( const FString& Name, const pxr::VtArray< pxr::GfVec3f >& PointOffsets, const pxr::VtArray< pxr::GfVec3f >& NormalOffsets, const pxr::VtArray< int >& PointIndices, const FUsdStageInfo& StageInfo, const FTransform& AdditionalTransform, uint32 PointIndexOffset, int32 LODIndex, UsdUtils::FUsdBlendShape& OutBlendShape )
 	{
 		uint32 NumOffsets = PointOffsets.size();
@@ -304,8 +304,8 @@ namespace SkelDataConversionImpl
 		OutBlendShape.Vertices.SetNumUninitialized( NumOffsets );
 		for ( uint32 OffsetIndex = 0; OffsetIndex < NumOffsets; ++OffsetIndex )
 		{
-			const FVector UE4Offset = UsdToUnreal::ConvertVector( StageInfo, PointOffsets[ OffsetIndex ] );
-			const FVector UE4Normal = OutBlendShape.bHasAuthoredTangents
+			const FVector UEOffset = UsdToUnreal::ConvertVector( StageInfo, PointOffsets[ OffsetIndex ] );
+			const FVector UENormal = OutBlendShape.bHasAuthoredTangents
 				? UsdToUnreal::ConvertVector( StageInfo, NormalOffsets[ OffsetIndex ] )
 				: FVector( 0, 0, 0 );
 
@@ -313,8 +313,8 @@ namespace SkelDataConversionImpl
 
 			// Intentionally ignore translation on PositionDelta as this is really a direction vector,
 			// and geomBindTransform's translation is already applied to the mesh vertices
-			ModifiedVertex.PositionDelta = AdditionalTransform.TransformVector(UE4Offset);
-			ModifiedVertex.TangentZDelta = NormalTransform.TransformVector(UE4Normal);
+			ModifiedVertex.PositionDelta = AdditionalTransform.TransformVector(UEOffset);
+			ModifiedVertex.TangentZDelta = NormalTransform.TransformVector(UENormal);
 			ModifiedVertex.SourceIdx = BaseIndices[ OffsetIndex ];
 		}
 
@@ -1731,7 +1731,7 @@ bool UsdToUnreal::ConvertSkelAnim( const pxr::UsdSkelSkeletonQuery& InUsdSkeleto
 
 		for ( int32 BoneIndex = 0; BoneIndex < NumBones; ++BoneIndex )
 		{
-			Controller->AddBoneTrack(BoneInfo[BoneIndex].Name); 
+			Controller->AddBoneTrack(BoneInfo[BoneIndex].Name);
 			Controller->SetBoneTrackKeys(BoneInfo[BoneIndex].Name, JointTracks[BoneIndex].PosKeys, JointTracks[BoneIndex].RotKeys, JointTracks[BoneIndex].ScaleKeys);
 		}
 	}
@@ -2681,7 +2681,7 @@ bool UnrealToUsd::ConvertAnimSequence( UAnimSequence* AnimSequence, pxr::UsdPrim
 			ScalesAttr.Set( Scales, pxr::UsdTimeCode( TimeCode ) );
 		}
 	}
-	
+
 	const int32 StageEndTimeCode = SkelAnimPrim.GetStage()->GetEndTimeCode();
 
 	if ( NumTimeCodes > StageEndTimeCode )
diff --git a/Engine/Plugins/Importers/USDImporter/Source/USDUtilities/Public/USDGeomMeshConversion.h b/Engine/Plugins/Importers/USDImporter/Source/USDUtilities/Public/USDGeomMeshConversion.h
--- a/Engine/Plugins/Importers/USDImporter/Source/USDUtilities/Public/USDGeomMeshConversion.h
+++ b/Engine/Plugins/Importers/USDImporter/Source/USDUtilities/Public/USDGeomMeshConversion.h
@@ -40,7 +40,7 @@ namespace UsdToUnreal
 	 * @param UsdSchema - UsdGeomMesh that contains the data to convert
 	 * @param MeshDescription - Output parameter that will be filled with the converted data
 	 * @param MaterialAssignments - Output parameter that will be filled with the material data extracted from UsdSchema
-	 * @param AdditionalTransform - Transform (in UE4 space) to bake into the vertex positions and normals of the converted mesh
+	 * @param AdditionalTransform - Transform (in UnrealEditor space) to bake into the vertex positions and normals of the converted mesh
 	 * @param MaterialToPrimvarsUVSetNames - Maps from a material prim path, to pairs indicating which primvar names are used as 'st' coordinates for this mesh, and which UVIndex materials will sample from (e.g. ["st0", 0], ["myUvSet2", 2], etc). This is used to pick which primvars will become UV sets.
 	 * @param TimeCode - USD timecode when the UsdGeomMesh should be sampled for mesh data and converted
 	 * @return Whether the conversion was successful or not.
