Disabled clang warnings generated by [REDACTED_PERSON]

preflight 6079ab8b6f85a8000169c380

diff --git a/Engine/Plugins/Enterprise/DatasmithC4DImporter/Source/DatasmithC4DTranslator/DatasmithC4DTranslator.Build.cs b/Engine/Plugins/Enterprise/DatasmithC4DImporter/Source/DatasmithC4DTranslator/DatasmithC4DTranslator.Build.cs
--- a/Engine/Plugins/Enterprise/DatasmithC4DImporter/Source/DatasmithC4DTranslator/DatasmithC4DTranslator.Build.cs
+++ b/Engine/Plugins/Enterprise/DatasmithC4DImporter/Source/DatasmithC4DTranslator/DatasmithC4DTranslator.Build.cs
@@ -61,6 +61,11 @@ namespace UnrealBuildTool.Rules
 				PublicDefinitions.Add("_MELANGE_SDK_");
 				PrivateDependencyModuleNames.Add("MelangeSDK");
 			}
+
+			if (Target.WindowsPlatform.Compiler == WindowsCompiler.Clang)
+			{
+				PublicDefinitions.Add("WITH_CLANG_COMPILER");
+			}
 		}
 	}
 }
diff --git a/Engine/Plugins/Enterprise/DatasmithC4DImporter/Source/DatasmithC4DTranslator/Private/DatasmithC4DMelangeSDKEnterGuard.h b/Engine/Plugins/Enterprise/DatasmithC4DImporter/Source/DatasmithC4DTranslator/Private/DatasmithC4DMelangeSDKEnterGuard.h
--- a/Engine/Plugins/Enterprise/DatasmithC4DImporter/Source/DatasmithC4DTranslator/Private/DatasmithC4DMelangeSDKEnterGuard.h
+++ b/Engine/Plugins/Enterprise/DatasmithC4DImporter/Source/DatasmithC4DTranslator/Private/DatasmithC4DMelangeSDKEnterGuard.h
@@ -19,5 +19,9 @@ __pragma(warning(disable: 6297)) /* melange\20.004_rbmelange20.0_259890\includes
 
 #endif // PLATFORM_WINDOWS
 
+#ifdef WITH_CLANG_COMPILER
+_Pragma( "clang diagnostic push" )
+_Pragma( "clang diagnostic ignored \"-Wdeprecated-declarations\"" )
+#endif
 
 
diff --git a/Engine/Plugins/Enterprise/DatasmithC4DImporter/Source/DatasmithC4DTranslator/Private/DatasmithC4DMelangeSDKLeaveGuard.h b/Engine/Plugins/Enterprise/DatasmithC4DImporter/Source/DatasmithC4DTranslator/Private/DatasmithC4DMelangeSDKLeaveGuard.h
--- a/Engine/Plugins/Enterprise/DatasmithC4DImporter/Source/DatasmithC4DTranslator/Private/DatasmithC4DMelangeSDKLeaveGuard.h
+++ b/Engine/Plugins/Enterprise/DatasmithC4DImporter/Source/DatasmithC4DTranslator/Private/DatasmithC4DMelangeSDKLeaveGuard.h
@@ -1,5 +1,9 @@
 // Copyright Epic Games, Inc. All Rights Reserved.
 
+#ifdef WITH_CLANG_COMPILER
+_Pragma( "clang diagnostic pop" )
+#endif
+
 #if PLATFORM_WINDOWS
 
 __pragma(warning(pop))
