Fix for LWC build
#rnx

diff --git a/Engine/Plugins/Enterprise/DatasmithC4DImporter/Source/DatasmithC4DDynamicImporter/Private/DatasmithC4DImporter.cpp b/Engine/Plugins/Enterprise/DatasmithC4DImporter/Source/DatasmithC4DDynamicImporter/Private/DatasmithC4DImporter.cpp
--- a/Engine/Plugins/Enterprise/DatasmithC4DImporter/Source/DatasmithC4DDynamicImporter/Private/DatasmithC4DImporter.cpp
+++ b/Engine/Plugins/Enterprise/DatasmithC4DImporter/Source/DatasmithC4DDynamicImporter/Private/DatasmithC4DImporter.cpp
@@ -331,7 +331,7 @@ FTransform CalculateCraneCameraTransform(const FCraneCameraAttributes& Params)
 	// More specifically, convert them so that that ConvertDirectionLeftHandedYup and
 	// the conversion for Ocamera rotations gets them back into UE4's coordinate system
 	// Note: Remember that FRotator's constructor is Pitch, Yaw and Roll (i.e. Y, Z, X)
-	return FTransform(FRotator(EulerUE4.Y, EulerUE4.X, -EulerUE4.Z - 90),
+	return FTransform(FRotator((float)EulerUE4.Y, (float)EulerUE4.X, (float)(-EulerUE4.Z - 90)),
 		FVector(TranslationUE4.X, TranslationUE4.Z, -TranslationUE4.Y));
 }
 
@@ -1691,7 +1691,7 @@ TSharedPtr<IDatasmithMasterMaterialElement> FDatasmithC4DDynamicImporter::Import
 			FVector TransparencyColor = MelangeGetVector(InC4DMaterialPtr, MATERIAL_TRANSPARENCY_COLOR);
 
 			// In Cinema4D Transparency Color seems to be used just as another multiplier for the opacity, not as an actual color
-			AddFloatToMaterial(MaterialPtr, TEXT("Transparency_Amount"), BrightnessValue * TransparencyColor.X * TransparencyColor.Y * TransparencyColor.Z);
+			AddFloatToMaterial(MaterialPtr, TEXT("Transparency_Amount"), BrightnessValue * (float)TransparencyColor.X * (float)TransparencyColor.Y * (float)TransparencyColor.Z);
 		}
 
 		float TransparencyRefraction = MelangeGetFloat(InC4DMaterialPtr, MATERIAL_TRANSPARENCY_REFRACTION);
@@ -2225,7 +2225,7 @@ namespace
 				}
 				else
 				{
-					Value /= InitialSize[TransformVectorIndex];
+					Value /= (float)InitialSize[TransformVectorIndex];
 				}
 			}
 		}
@@ -2792,9 +2792,9 @@ void FDatasmithC4DDynamicImporter::ImportAnimations(TSharedPtr<IDatasmithActorEl
 
 				// TransformValue represents, in radians, the rotations around the C4D axes
 				// XRot, YRot, ZRot are rotations around UE4 axes, in the UE4 CS, with the sign given by Quaternion rotations (NOT Rotators)
-				FQuat XRot = FQuat(FVector(1, 0, 0), -TransformValueCopy.X);
-				FQuat YRot = FQuat(FVector(0, 1, 0), TransformValueCopy.Z);
-				FQuat ZRot = FQuat(FVector(0, 0, 1), -TransformValueCopy.Y);
+				FQuat XRot = FQuat(FVector(1, 0, 0), (float)(-TransformValueCopy.X));
+				FQuat YRot = FQuat(FVector(0, 1, 0), (float)(TransformValueCopy.Z));
+				FQuat ZRot = FQuat(FVector(0, 0, 1), (float)(-TransformValueCopy.Y));
 
 				// Swap YRot and ZRot in the composition order, as an XYZ order in the C4D CS really means a XZY order in the UE4 CS
 				// This effectively converts the rotation order from the C4D CS to the UE4 CS, the sign of the rotations being handled when
@@ -3055,9 +3055,9 @@ void FDatasmithC4DDynamicImporter::ImportDrivenAnimations(TSharedPtr<IDatasmithA
 
 				// TransformValue represents, in radians, the rotations around the C4D axes
 				// XRot, YRot, ZRot are rotations around UE4 axes, in the UE4 CS, with the sign given by Quaternion rotations (NOT Rotators)
-				FQuat XRot = FQuat(FVector(1, 0, 0), -TransformValueCopy.X);
-				FQuat YRot = FQuat(FVector(0, 1, 0), TransformValueCopy.Z);
-				FQuat ZRot = FQuat(FVector(0, 0, 1), -TransformValueCopy.Y);
+				FQuat XRot = FQuat(FVector(1, 0, 0), (float)(-TransformValueCopy.X));
+				FQuat YRot = FQuat(FVector(0, 1, 0), (float)(TransformValueCopy.Z));
+				FQuat ZRot = FQuat(FVector(0, 0, 1), (float)(-TransformValueCopy.Y));
 
 				// Swap YRot and ZRot in the composition order, as an XYZ order in the C4D CS really means a XZY order in the UE4 CS
 				// This effectively converts the rotation order from the C4D CS to the UE4 CS, the sign of the rotations being handled when
@@ -3654,8 +3654,8 @@ TSharedPtr<IDatasmithMeshElement> FDatasmithC4DDynamicImporter::ImportMesh(cinew
 	MeshDescription.Empty();
 
 	FStaticMeshAttributes StaticMeshAttributes(MeshDescription);
-	TVertexAttributesRef<FVector> VertexPositions = StaticMeshAttributes.GetVertexPositions();
-	TVertexInstanceAttributesRef<FVector> VertexInstanceNormals = StaticMeshAttributes.GetVertexInstanceNormals();
+	TVertexAttributesRef<FVector3f> VertexPositions = StaticMeshAttributes.GetVertexPositions();
+	TVertexInstanceAttributesRef<FVector3f> VertexInstanceNormals = StaticMeshAttributes.GetVertexInstanceNormals();
 	TVertexInstanceAttributesRef<FVector2D> VertexInstanceUVs = StaticMeshAttributes.GetVertexInstanceUVs();
 	TPolygonGroupAttributesRef<FName> PolygonGroupImportedMaterialSlotNames = StaticMeshAttributes.GetPolygonGroupMaterialSlotNames();
 
