Fixed issue with DebugGame configuration

diff --git a/Engine/Plugins/Enterprise/DatasmithC4DImporter/Source/DatasmithC4DDynamicImporter/DatasmithC4DDynamicImporter.Build.cs b/Engine/Plugins/Enterprise/DatasmithC4DImporter/Source/DatasmithC4DDynamicImporter/DatasmithC4DDynamicImporter.Build.cs
--- a/Engine/Plugins/Enterprise/DatasmithC4DImporter/Source/DatasmithC4DDynamicImporter/DatasmithC4DDynamicImporter.Build.cs
+++ b/Engine/Plugins/Enterprise/DatasmithC4DImporter/Source/DatasmithC4DDynamicImporter/DatasmithC4DDynamicImporter.Build.cs
@@ -71,7 +71,7 @@ namespace UnrealBuildTool.Rules
 			// Make sure the C4D Cineware SDK can be used to compile
 			bool bCanUseCinewareSDK = Directory.Exists(CinewareSDKLocation);
 			// Temporary: Debug configuration does not link yet
-			bCanUseCinewareSDK &= Target.Configuration == UnrealTargetConfiguration.Development || Target.Configuration == UnrealTargetConfiguration.Shipping;
+			bCanUseCinewareSDK &= Target.Configuration != UnrealTargetConfiguration.Unknown && Target.Configuration != UnrealTargetConfiguration.Debug;
 			// Temporary: Do not enable Cineware SDK code if static analysis is requested
 			if (Target.Platform == UnrealTargetPlatform.Win64)
             {
