Fixed build warning

diff --git a/Engine/Plugins/Enterprise/DatasmithC4DImporter/Source/DatasmithC4DDynamicImporter/Private/DatasmithC4DImporterMaterialSelector.h b/Engine/Plugins/Enterprise/DatasmithC4DImporter/Source/DatasmithC4DDynamicImporter/Private/DatasmithC4DImporterMaterialSelector.h
--- a/Engine/Plugins/Enterprise/DatasmithC4DImporter/Source/DatasmithC4DDynamicImporter/Private/DatasmithC4DImporterMaterialSelector.h
+++ b/Engine/Plugins/Enterprise/DatasmithC4DImporter/Source/DatasmithC4DDynamicImporter/Private/DatasmithC4DImporterMaterialSelector.h
@@ -2,6 +2,8 @@
 
 #pragma once
 
+#include "CoreMinimal.h"
+
 #include "DatasmithDefinitions.h"
 #include "MasterMaterials/DatasmithMasterMaterial.h"
 #include "MasterMaterials/DatasmithMasterMaterialSelector.h"
