C4D DynamicImporter update for R25 SDK and bugfixes (Contributed by [REDACTED_PERSON])

#8856

diff --git a/Engine/Plugins/Enterprise/DatasmithC4DImporter/Source/DatasmithC4DDynamicImporter/DatasmithC4DDynamicImporter.Build.cs b/Engine/Plugins/Enterprise/DatasmithC4DImporter/Source/DatasmithC4DDynamicImporter/DatasmithC4DDynamicImporter.Build.cs
--- a/Engine/Plugins/Enterprise/DatasmithC4DImporter/Source/DatasmithC4DDynamicImporter/DatasmithC4DDynamicImporter.Build.cs
+++ b/Engine/Plugins/Enterprise/DatasmithC4DImporter/Source/DatasmithC4DDynamicImporter/DatasmithC4DDynamicImporter.Build.cs
@@ -113,7 +113,7 @@ namespace UnrealBuildTool.Rules
 					PublicDefinitions.Add("MAXON_TARGET_WINDOWS");
 
 
-					PublicDefinitions.Add("CINEWARE_LOCATION=\"C:/Program Files/Maxon Cinema 4D R24/cineware.dll\"");
+					PublicDefinitions.Add("CINEWARE_LOCATION=\"C:/Program Files/Maxon Cinema 4D R25/cineware.dll\"");
 
 					PublicAdditionalLibraries.Add(Path.Combine(CinewareSDKLocation, "libraries", "cinema.framework_Release_64bit.lib"));
 					PublicAdditionalLibraries.Add(Path.Combine(CinewareSDKLocation, "libraries", "core.framework_Release_64bit.lib"));
@@ -126,7 +126,7 @@ namespace UnrealBuildTool.Rules
 					PublicDefinitions.Add("MAXON_TARGET_MACOS");
 					PublicDefinitions.Add("MAXON_TARGET_OSX");
 
-					PublicDefinitions.Add("CINEWARE_LOCATION=\"/Applications/Maxon Cinema 4D R24/cineware.bundle/Contents/MacOS/cineware\"");
+					PublicDefinitions.Add("CINEWARE_LOCATION=\"/Applications/Maxon Cinema 4D R25/cineware.bundle/Contents/MacOS/cineware\"");
 
 					PublicAdditionalLibraries.Add(Path.Combine(CinewareSDKLocation, "libraries", "libcinema.framework.a"));
 					PublicAdditionalLibraries.Add(Path.Combine(CinewareSDKLocation, "libraries", "libcore.framework.a"));
diff --git a/Engine/Plugins/Enterprise/DatasmithC4DImporter/Source/DatasmithC4DDynamicImporter/Private/DatasmithC4DDynamicImporterModule.cpp b/Engine/Plugins/Enterprise/DatasmithC4DImporter/Source/DatasmithC4DDynamicImporter/Private/DatasmithC4DDynamicImporterModule.cpp
--- a/Engine/Plugins/Enterprise/DatasmithC4DImporter/Source/DatasmithC4DDynamicImporter/Private/DatasmithC4DDynamicImporterModule.cpp
+++ b/Engine/Plugins/Enterprise/DatasmithC4DImporter/Source/DatasmithC4DDynamicImporter/Private/DatasmithC4DDynamicImporterModule.cpp
@@ -97,6 +97,13 @@ public:
 		FString EnvVariable = FPlatformMisc::GetEnvironmentVariable(TEXT("DATASMITHC4D_DEBUG"));
 		bDebugMode = !EnvVariable.IsEmpty();
 	}
+	
+	virtual void ShutdownModule() override
+	{
+#ifdef _CINEWARE_SDK_
+		cineware::UnloadCineware();
+#endif
+	}
 
 #ifdef _CINEWARE_SDK_
 	FReply RedirectToEndpoint(const FGeometry& Geometry, const FPointerEvent& Event)
diff --git a/Engine/Plugins/Enterprise/DatasmithC4DImporter/Source/DatasmithC4DDynamicImporter/Private/DatasmithC4DExtraMelangeDefinitions.h b/Engine/Plugins/Enterprise/DatasmithC4DImporter/Source/DatasmithC4DDynamicImporter/Private/DatasmithC4DExtraMelangeDefinitions.h
--- a/Engine/Plugins/Enterprise/DatasmithC4DImporter/Source/DatasmithC4DDynamicImporter/Private/DatasmithC4DExtraMelangeDefinitions.h
+++ b/Engine/Plugins/Enterprise/DatasmithC4DImporter/Source/DatasmithC4DDynamicImporter/Private/DatasmithC4DExtraMelangeDefinitions.h
@@ -54,6 +54,9 @@ namespace cineware
 		ID_BASEOBJECT_REL_POSITION = 903,
 		ID_BASEOBJECT_REL_ROTATION = 904,
 		ID_BASEOBJECT_REL_SCALE = 905,
+		MG_LINEAR_OBJECT_POSITION = 1273,
+		MG_LINEAR_OBJECT_ROTATION = 1274,
+		MG_LINEAR_OBJECT_SCALE = 1275,
 	};
 
 	enum // Ocloner (defined in resources folder)
diff --git a/Engine/Plugins/Enterprise/DatasmithC4DImporter/Source/DatasmithC4DDynamicImporter/Private/DatasmithC4DImporter.cpp b/Engine/Plugins/Enterprise/DatasmithC4DImporter/Source/DatasmithC4DDynamicImporter/Private/DatasmithC4DImporter.cpp
--- a/Engine/Plugins/Enterprise/DatasmithC4DImporter/Source/DatasmithC4DDynamicImporter/Private/DatasmithC4DImporter.cpp
+++ b/Engine/Plugins/Enterprise/DatasmithC4DImporter/Source/DatasmithC4DDynamicImporter/Private/DatasmithC4DImporter.cpp
@@ -2059,7 +2059,7 @@ TSharedPtr<IDatasmithMeshActorElement> FDatasmithC4DDynamicImporter::ImportPolyg
 	FString HashString = BytesToHex(PolygonHash.GetBytes(), PolygonHash.GetSize());
 
 	TOptional<FString> MelangeID = MelangeObjectID(PolyObject);
-	if (!MelangeID.IsSet())
+	if (!MelangeID.IsSet() || (PolyObject && PolyObject->GetPointCount() < 3))
 	{
 		return nullptr;
 	}
@@ -2148,7 +2148,7 @@ TSharedPtr<IDatasmithMeshActorElement> FDatasmithC4DDynamicImporter::ImportPolyg
 		MeshActorElement->AddMaterialOverride(MaterialIDElement);
 
 		// If we have an unassigned material, we *must* set it on the base mesh, as we can't create a material overrides to "clear" a material slot
-		if (!bMeshHasMaterialAssignments || MaterialName.IsEmpty())
+		if (!bMeshHasMaterialAssignments)
 		{
 			ResultMeshElement->SetMaterial(*MaterialName, SlotIndex);
 		}
@@ -2844,16 +2844,10 @@ void FDatasmithC4DDynamicImporter::ImportAnimations(TSharedPtr<IDatasmithActorEl
 void FDatasmithC4DDynamicImporter::ImportDrivenAnimations(TSharedPtr<IDatasmithActorElement> ActorElement, cineware::Int32 FrameNumber)
 {
 	bool IsAddedNull = false;
-	cineware::BaseObject* Object = *ActorElementToAnimationSources.Find(ActorElement.Get());
+	if (!ActorElement.IsValid())
+		return;
 
-	/*for (int32 TagIndex = 0; TagIndex < ActorElement->GetTagsCount(); ++TagIndex)
-	{
-		if (FString(ActorElement->GetTag(TagIndex)) == FString("AddedNull"))
-		{
-			IsAddedNull = true;
-			break;
-		}
-	}*/
+	cineware::BaseObject* Object = *ActorElementToAnimationSources.Find(ActorElement.Get());
 
 	// an object might be not exist for the actual frame
 	if (Object == nullptr || IsAddedNull)
@@ -2998,15 +2992,15 @@ void FDatasmithC4DDynamicImporter::ImportDrivenAnimations(TSharedPtr<IDatasmithA
 	bool needsToBeAdded = false;
 	TSharedPtr<IDatasmithTransformAnimationElement> Animation = nullptr;
 
-	if (BaseObjectImportedAnimationElement.Contains(Object))
+	if (BaseObjectImportedAnimationElement.Contains(ActorElement->GetName()))
 	{
-		Animation = *BaseObjectImportedAnimationElement.Find(Object);
+		Animation = *BaseObjectImportedAnimationElement.Find(ActorElement->GetName());
 		needsToBeAdded = false;
 	}
 	else
 	{
 		Animation = FDatasmithSceneFactory::CreateTransformAnimation(ActorElement->GetName());
-		BaseObjectImportedAnimationElement.Add(Object, Animation);
+		BaseObjectImportedAnimationElement.Add(ActorElement->GetName(), Animation);
 		needsToBeAdded = true;
 	}
 
@@ -3661,7 +3655,7 @@ TSharedPtr<IDatasmithMeshElement> FDatasmithC4DDynamicImporter::ImportMesh(cinew
 
 	// Reserve space for attributes. These might not be enough as some of these polygons might be quads or n-gons, but its better than nothing
 	MeshDescription.ReserveNewVertices(PointCount);
-	MeshDescription.ReserveNewVertexInstances(PolygonCount);
+	MeshDescription.ReserveNewVertexInstances(PointCount);
 	MeshDescription.ReserveNewEdges(PolygonCount);
 	MeshDescription.ReserveNewPolygons(PolygonCount);
 	MeshDescription.ReserveNewPolygonGroups(NumSlots);
diff --git a/Engine/Plugins/Enterprise/DatasmithC4DImporter/Source/DatasmithC4DDynamicImporter/Private/DatasmithC4DImporter.h b/Engine/Plugins/Enterprise/DatasmithC4DImporter/Source/DatasmithC4DDynamicImporter/Private/DatasmithC4DImporter.h
--- a/Engine/Plugins/Enterprise/DatasmithC4DImporter/Source/DatasmithC4DDynamicImporter/Private/DatasmithC4DImporter.h
+++ b/Engine/Plugins/Enterprise/DatasmithC4DImporter/Source/DatasmithC4DDynamicImporter/Private/DatasmithC4DImporter.h
@@ -203,7 +203,7 @@ private:
 	TMap<cineware::BaseObject*, TMap<int32, TMap<EDatasmithTransformType, FVector>>> BaseObjectImportedTransformType;
 
 	/** Storage of all parsed objects from the melange document, used so we can import all animations afterwards */
-	TMap<cineware::BaseObject*, TSharedPtr<IDatasmithTransformAnimationElement>> BaseObjectImportedAnimationElement;
+	TMap<FString, TSharedPtr<IDatasmithTransformAnimationElement>> BaseObjectImportedAnimationElement;
 
 	/** Maps an instance to the corresponding original node, used so that we can redirect animations to the original nodes */
 	TMap<cineware::BaseObject*, cineware::BaseObject*> InstancedSubObjectsToOriginals;
