Fixed build break following changes in Cineware SDK

diff --git a/Engine/Plugins/Enterprise/DatasmithC4DImporter/Source/DatasmithC4DDynamicImporter/Private/DatasmithC4DImporter.cpp b/Engine/Plugins/Enterprise/DatasmithC4DImporter/Source/DatasmithC4DDynamicImporter/Private/DatasmithC4DImporter.cpp
--- a/Engine/Plugins/Enterprise/DatasmithC4DImporter/Source/DatasmithC4DDynamicImporter/Private/DatasmithC4DImporter.cpp
+++ b/Engine/Plugins/Enterprise/DatasmithC4DImporter/Source/DatasmithC4DDynamicImporter/Private/DatasmithC4DImporter.cpp
@@ -58,8 +58,6 @@ DECLARE_CYCLE_STAT(TEXT("C4DImporter - Load File"), STAT_C4DImporter_LoadFile, S
 #define UnitlessGlobalLightIntensity 10.0
 #define UnitlessIESandPointLightIntensity 8000
 
-class BaseMaterial;
-
 FDatasmithC4DDynamicImporter::FPreTranslateEvent FDatasmithC4DDynamicImporter::PreTranslateEvent;
 
 FDatasmithC4DDynamicImporter::FDatasmithC4DDynamicImporter(TSharedRef<IDatasmithScene>& OutScene, FDatasmithC4DImportOptions& InOptions)
