Speculative fix to fix aliasing for Cineware SDK name collision

diff --git a/Engine/Plugins/Enterprise/DatasmithC4DImporter/Source/DatasmithC4DDynamicImporter/Private/DatasmithC4DImporter.h b/Engine/Plugins/Enterprise/DatasmithC4DImporter/Source/DatasmithC4DDynamicImporter/Private/DatasmithC4DImporter.h
--- a/Engine/Plugins/Enterprise/DatasmithC4DImporter/Source/DatasmithC4DDynamicImporter/Private/DatasmithC4DImporter.h
+++ b/Engine/Plugins/Enterprise/DatasmithC4DImporter/Source/DatasmithC4DDynamicImporter/Private/DatasmithC4DImporter.h
@@ -21,6 +21,12 @@ DECLARE_LOG_CATEGORY_EXTERN(LogDatasmithC4DImport, Log, All)
 
 DATASMITH_C4D_PUSH_WARNINGS
 #include "cineware.h"
+
+// "Describe" macro is colliding with any function with the same name
+#if defined(Describe)
+#undef Describe
+#endif
+
 DATASMITH_C4D_POP_WARNINGS
 
 class FDatasmithSceneExporter;
diff --git a/Engine/Plugins/Enterprise/DatasmithC4DImporter/Source/DatasmithC4DDynamicImporter/Private/DatasmithC4DUtils.h b/Engine/Plugins/Enterprise/DatasmithC4DImporter/Source/DatasmithC4DDynamicImporter/Private/DatasmithC4DUtils.h
--- a/Engine/Plugins/Enterprise/DatasmithC4DImporter/Source/DatasmithC4DDynamicImporter/Private/DatasmithC4DUtils.h
+++ b/Engine/Plugins/Enterprise/DatasmithC4DImporter/Source/DatasmithC4DDynamicImporter/Private/DatasmithC4DUtils.h
@@ -65,6 +65,11 @@ DATASMITH_C4D_PUSH_WARNINGS
 #undef IsMinimized
 #endif
 
+// "Describe" macro is colliding with any function with the same name
+#if defined(Describe)
+#undef Describe
+#endif
+
 DATASMITH_C4D_POP_WARNINGS
 /**
  * Retrieves the value of a DA_LONG parameter of a melange object as a cineware::Int32,
