PR #9349: Removed DatasmithC4DDynamicImporter module, added Notification for improved cineware importer (Contributed by [REDACTED_PERSON])

Reason: Maxon is regaining ownership of c4d import using the Cineware SDK.

#9349

diff --git a/Engine/Plugins/Enterprise/DatasmithC4DImporter/DatasmithC4DImporter.uplugin b/Engine/Plugins/Enterprise/DatasmithC4DImporter/DatasmithC4DImporter.uplugin
--- a/Engine/Plugins/Enterprise/DatasmithC4DImporter/DatasmithC4DImporter.uplugin
+++ b/Engine/Plugins/Enterprise/DatasmithC4DImporter/DatasmithC4DImporter.uplugin
@@ -1,7 +1,7 @@
 {
 	"FileVersion" : 3,
-	"Version" : 1,
-	"VersionName" : "1.0",
+	"Version" : 2,
+	"VersionName" : "1.1",
 	"FriendlyName" : "Datasmith C4D Importer",
 	"Description" : "Adds support for importing content from Cinema4D applications into Unreal Engine",
 	"Category" : "Importers",
@@ -30,16 +30,7 @@
 				"Win64",
 				"Mac"
 			]
-		},
-		{
-			"Name": "DatasmithC4DDynamicImporter",
-			"Type": "Editor",
-			"LoadingPhase": "PostDefault",
-			"PlatformAllowList": [
-				"Win64",
-				"Mac"
-			]
-		}
+		}		
 	],
 	"Plugins":
 	[
@@ -47,5 +38,5 @@
 			"Name": "DatasmithImporter",
 			"Enabled": true
 		}
-  ]
+	]
 }
diff --git a/Engine/Plugins/Enterprise/DatasmithC4DImporter/Source/DatasmithC4DDynamicImporter/DatasmithC4DDynamicImporter.Build.cs b/Engine/Plugins/Enterprise/DatasmithC4DImporter/Source/DatasmithC4DDynamicImporter/DatasmithC4DDynamicImporter.Build.cs
deleted file mode 100644
diff --git a/Engine/Plugins/Enterprise/DatasmithC4DImporter/Source/DatasmithC4DDynamicImporter/Private/DatasmithC4DDynamicImporterModule.cpp b/Engine/Plugins/Enterprise/DatasmithC4DImporter/Source/DatasmithC4DDynamicImporter/Private/DatasmithC4DDynamicImporterModule.cpp
deleted file mode 100644
diff --git a/Engine/Plugins/Enterprise/DatasmithC4DImporter/Source/DatasmithC4DDynamicImporter/Private/DatasmithC4DExtraMelangeDefinitions.h b/Engine/Plugins/Enterprise/DatasmithC4DImporter/Source/DatasmithC4DDynamicImporter/Private/DatasmithC4DExtraMelangeDefinitions.h
deleted file mode 100644
diff --git a/Engine/Plugins/Enterprise/DatasmithC4DImporter/Source/DatasmithC4DDynamicImporter/Private/DatasmithC4DImporter.cpp b/Engine/Plugins/Enterprise/DatasmithC4DImporter/Source/DatasmithC4DDynamicImporter/Private/DatasmithC4DImporter.cpp
deleted file mode 100644
diff --git a/Engine/Plugins/Enterprise/DatasmithC4DImporter/Source/DatasmithC4DDynamicImporter/Private/DatasmithC4DImporter.h b/Engine/Plugins/Enterprise/DatasmithC4DImporter/Source/DatasmithC4DDynamicImporter/Private/DatasmithC4DImporter.h
deleted file mode 100644
diff --git a/Engine/Plugins/Enterprise/DatasmithC4DImporter/Source/DatasmithC4DDynamicImporter/Private/DatasmithC4DImporterMaterialSelector.cpp b/Engine/Plugins/Enterprise/DatasmithC4DImporter/Source/DatasmithC4DDynamicImporter/Private/DatasmithC4DImporterMaterialSelector.cpp
deleted file mode 100644
diff --git a/Engine/Plugins/Enterprise/DatasmithC4DImporter/Source/DatasmithC4DDynamicImporter/Private/DatasmithC4DImporterMaterialSelector.h b/Engine/Plugins/Enterprise/DatasmithC4DImporter/Source/DatasmithC4DDynamicImporter/Private/DatasmithC4DImporterMaterialSelector.h
deleted file mode 100644
diff --git a/Engine/Plugins/Enterprise/DatasmithC4DImporter/Source/DatasmithC4DDynamicImporter/Private/DatasmithC4DMelangeSDKOverload.cpp b/Engine/Plugins/Enterprise/DatasmithC4DImporter/Source/DatasmithC4DDynamicImporter/Private/DatasmithC4DMelangeSDKOverload.cpp
deleted file mode 100644
diff --git a/Engine/Plugins/Enterprise/DatasmithC4DImporter/Source/DatasmithC4DDynamicImporter/Private/DatasmithC4DUtils.cpp b/Engine/Plugins/Enterprise/DatasmithC4DImporter/Source/DatasmithC4DDynamicImporter/Private/DatasmithC4DUtils.cpp
deleted file mode 100644
diff --git a/Engine/Plugins/Enterprise/DatasmithC4DImporter/Source/DatasmithC4DDynamicImporter/Private/DatasmithC4DUtils.h b/Engine/Plugins/Enterprise/DatasmithC4DImporter/Source/DatasmithC4DDynamicImporter/Private/DatasmithC4DUtils.h
deleted file mode 100644
diff --git a/Engine/Plugins/Enterprise/DatasmithC4DImporter/Source/DatasmithC4DDynamicImporter/Public/DatasmithC4DDynamicImporterModule.h b/Engine/Plugins/Enterprise/DatasmithC4DImporter/Source/DatasmithC4DDynamicImporter/Public/DatasmithC4DDynamicImporterModule.h
deleted file mode 100644
diff --git a/Engine/Plugins/Enterprise/DatasmithC4DImporter/Source/DatasmithC4DTranslator/DatasmithC4DTranslator.Build.cs b/Engine/Plugins/Enterprise/DatasmithC4DImporter/Source/DatasmithC4DTranslator/DatasmithC4DTranslator.Build.cs
--- a/Engine/Plugins/Enterprise/DatasmithC4DImporter/Source/DatasmithC4DTranslator/DatasmithC4DTranslator.Build.cs
+++ b/Engine/Plugins/Enterprise/DatasmithC4DImporter/Source/DatasmithC4DTranslator/DatasmithC4DTranslator.Build.cs
@@ -1,7 +1,6 @@
 // Copyright Epic Games, Inc. All Rights Reserved.
 
 using System.IO;
-using EpicGames.Core;
 
 namespace UnrealBuildTool.Rules
 {
@@ -19,11 +18,11 @@ namespace UnrealBuildTool.Rules
 					"CoreUObject",
 					"DatasmithCore",
 					"Engine",
-					"Json",
 					"MeshDescription",
 					"Slate",
 					"SlateCore",
 					"StaticMeshDescription",
+					"EditorStyle",
 					"Imath",
 				}
 			);
@@ -46,12 +45,7 @@ namespace UnrealBuildTool.Rules
 				}
 			);
 
-			string DsC4DDynamicLocatation = Path.Combine(PluginDirectory, "Source", "DatasmithC4DDynamicImporter");
-			if (Directory.Exists(DsC4DDynamicLocatation))
-			{
-				PublicIncludePaths.Add(Path.Combine(DsC4DDynamicLocatation, "Public"));
-				PublicDefinitions.Add("_CHECK_DYNAMIC_IMPORTER_");
-			}
+
 
 			// Set up the C4D Melange SDK includes and libraries.
 			string MelangeSDKLocation = Path.Combine(EngineDirectory, "Restricted/NotForLicensees/Source/ThirdParty/Enterprise/Melange/20.004_RBMelange20.0_259890");
diff --git a/Engine/Plugins/Enterprise/DatasmithC4DImporter/Source/DatasmithC4DTranslator/Private/DatasmithC4DImporter.cpp b/Engine/Plugins/Enterprise/DatasmithC4DImporter/Source/DatasmithC4DTranslator/Private/DatasmithC4DImporter.cpp
--- a/Engine/Plugins/Enterprise/DatasmithC4DImporter/Source/DatasmithC4DTranslator/Private/DatasmithC4DImporter.cpp
+++ b/Engine/Plugins/Enterprise/DatasmithC4DImporter/Source/DatasmithC4DTranslator/Private/DatasmithC4DImporter.cpp
@@ -3,7 +3,6 @@
 #ifdef _MELANGE_SDK_
 
 #include "DatasmithC4DImporter.h"
-#include "IDatasmithC4DImporter.h"
 
 #include "DatasmithAssetImportData.h"
 #include "DatasmithC4DExtraMelangeDefinitions.h"
diff --git a/Engine/Plugins/Enterprise/DatasmithC4DImporter/Source/DatasmithC4DTranslator/Private/DatasmithC4DImporter.h b/Engine/Plugins/Enterprise/DatasmithC4DImporter/Source/DatasmithC4DTranslator/Private/DatasmithC4DImporter.h
--- a/Engine/Plugins/Enterprise/DatasmithC4DImporter/Source/DatasmithC4DTranslator/Private/DatasmithC4DImporter.h
+++ b/Engine/Plugins/Enterprise/DatasmithC4DImporter/Source/DatasmithC4DTranslator/Private/DatasmithC4DImporter.h
@@ -11,7 +11,6 @@
 #include "UObject/StrongObjectPtr.h"
 
 #include "DatasmithTranslator.h"
-#include "IDatasmithC4DImporter.h"
 #include "DatasmithC4DMelangeSDKEnterGuard.h"
 #include "c4d.h"
 #include "c4d_browsecontainer.h"
@@ -36,27 +35,35 @@ struct FRichCurve;
 
 DECLARE_STATS_GROUP(TEXT("C4DImporter"), STATGROUP_C4DImporter, STATCAT_Advanced);
 
+struct FDatasmithC4DImportOptions
+{
+	bool bImportEmptyMesh = false;
+	bool bOptimizeEmptySingleChildActors = false;
+	bool bAlwaysGenerateNormals = false;
+	float ScaleVertices = 1.0;
+	bool bExportToUDatasmith = false;
+};
 
-class FDatasmithC4DImporter : public IDatasmithC4DImporter
+class FDatasmithC4DImporter
 {
 public:
 	FDatasmithC4DImporter(TSharedRef<IDatasmithScene>& OutScene, FDatasmithC4DImportOptions& InOptions);
-	virtual ~FDatasmithC4DImporter();
+	~FDatasmithC4DImporter();
 
 	/** Updates the used import options to InOptions */
-	virtual void SetImportOptions(FDatasmithC4DImportOptions& InOptions) override;
+	void SetImportOptions(FDatasmithC4DImportOptions& InOptions);
 
 	/** Open and load a .4cd file into C4dDocument */
-	virtual bool OpenFile(const FString& InFilename) override;
+	bool OpenFile(const FString& InFilename);
 
 	/** Parse the scene contained in the previously opened file and process its content according to parameters from incoming context */
-	virtual bool ProcessScene() override;
+	bool ProcessScene();
 
 	/** Unload melange resources after importing */
-	virtual void UnloadScene() override;
+	void UnloadScene();
 
-	virtual void GetGeometriesForMeshElementAndRelease(const TSharedRef<IDatasmithMeshElement> MeshElement, TArray<FMeshDescription>& OutMeshDescriptions) override;
-	virtual TSharedPtr<IDatasmithLevelSequenceElement> GetLevelSequence() override;
+	void GetGeometriesForMeshElementAndRelease(const TSharedRef<IDatasmithMeshElement> MeshElement, TArray<FMeshDescription>& OutMeshDescriptions);
+	TSharedPtr<IDatasmithLevelSequenceElement> GetLevelSequence();
 
 	/** Finds the most derived cache for a melange object. That will be e.g. a polygon cache or a deformed polygon cache, if it has one */
 	melange::BaseObject* GetBestMelangeCache(melange::BaseObject* Object);
@@ -114,15 +121,15 @@ public:
 	 * from InTextureTag, and adds the new material to the Datasmith scene. This is used because texture tags are closer to material
 	 * instances, and may have different "overrides" for each property
 	 */
-	virtual FString CustomizeMaterial(const FString& InMaterialID, const FString& InMeshID, melange::TextureTag* InTextureTag);
+	FString CustomizeMaterial(const FString& InMaterialID, const FString& InMeshID, melange::TextureTag* InTextureTag);
 
 	/**
 	 * Creates customized materials if necessary, and returns a map from material slot indices to material names
 	 */
-	virtual TMap<int32, FString> GetCustomizedMaterialAssignment(const FString& DatasmithMeshName, const TArray<melange::TextureTag*>& TextureTags);
+	TMap<int32, FString> GetCustomizedMaterialAssignment(const FString& DatasmithMeshName, const TArray<melange::TextureTag*>& TextureTags);
 
 	/** Imports a melange actor, which might involve parsing another small hierarchy of subnodes and deformers*/
-	virtual TSharedPtr<IDatasmithActorElement> ImportObjectAndChildren(melange::BaseObject* ActorObject, melange::BaseObject* DataObject, TSharedPtr<IDatasmithActorElement> ParentActor, const melange::Matrix& WorldTransformMatrix, const FString& InstancePath, const FString& DatasmithLabel, const TArray<melange::TextureTag*>& TextureTags);
+	TSharedPtr<IDatasmithActorElement> ImportObjectAndChildren(melange::BaseObject* ActorObject, melange::BaseObject* DataObject, TSharedPtr<IDatasmithActorElement> ParentActor, const melange::Matrix& WorldTransformMatrix, const FString& InstancePath, const FString& DatasmithLabel, const TArray<melange::TextureTag*>& TextureTags);
 
 	/** Traverse the melange actor hierarchy importing all nodes */
 	void ImportHierarchy(melange::BaseObject* ActorObject, melange::BaseObject* DataObject, TSharedPtr<IDatasmithActorElement> ParentActor, const melange::Matrix& WorldTransformMatrix, const FString& InstancePath, const TArray<melange::TextureTag*>& TextureTags);
diff --git a/Engine/Plugins/Enterprise/DatasmithC4DImporter/Source/DatasmithC4DTranslator/Private/DatasmithC4DTranslator.cpp b/Engine/Plugins/Enterprise/DatasmithC4DImporter/Source/DatasmithC4DTranslator/Private/DatasmithC4DTranslator.cpp
--- a/Engine/Plugins/Enterprise/DatasmithC4DImporter/Source/DatasmithC4DTranslator/Private/DatasmithC4DTranslator.cpp
+++ b/Engine/Plugins/Enterprise/DatasmithC4DImporter/Source/DatasmithC4DTranslator/Private/DatasmithC4DTranslator.cpp
@@ -1,20 +1,26 @@
 // Copyright Epic Games, Inc. All Rights Reserved.
 
 #include "DatasmithC4DTranslator.h"
-#include "DatasmithC4DDynamicImporterModule.h"
 
-#if defined(_MELANGE_SDK_) || defined(_CINEWARE_SDK_)
+#if defined(_MELANGE_SDK_)
 #include "DatasmithC4DTranslatorModule.h"
 #include "DatasmithSceneFactory.h"
 #include "DatasmithSceneSource.h"
 #include "IDatasmithSceneElements.h"
 
+#include "Framework/Notifications/NotificationManager.h"
 #include "GenericPlatform/GenericPlatformProcess.h"
 #include "Misc/FileHelper.h"
 #include "Misc/Paths.h"
+#include "Styling/AppStyle.h"
 #include "Templates/TypeHash.h"
+#include "Widgets/Notifications/SNotificationList.h"
 
-typedef TSharedPtr<IDatasmithC4DImporter>(*CreateC4DImporter)(TSharedRef<IDatasmithScene> OutScene, UDatasmithC4DImportOptions* InOptions);
+#define LOCTEXT_NAMESPACE "DatasmithC4DTranslator"
+
+DEFINE_LOG_CATEGORY_STATIC(DatasmithC4DTranslatorLog, Log, All);
+
+void NotifyImprovedImporter();
 
 void FDatasmithC4DTranslator::Initialize(FDatasmithTranslatorCapabilities& OutCapabilities)
 {
@@ -22,32 +28,16 @@ void FDatasmithC4DTranslator::Initialize(FDatasmithTranslatorCapabilities& OutCa
 	OutCapabilities.bParallelLoadStaticMeshSupported = true;
 
 	TArray<FFileFormatInfo>& Formats = OutCapabilities.SupportedFileFormats;
-    Formats.Emplace(TEXT("c4d"), TEXT("Cinema 4D file format"));
+	Formats.Emplace(TEXT("c4d"), TEXT("Cinema 4D file format"));
 }
 
 bool FDatasmithC4DTranslator::LoadScene(TSharedRef<IDatasmithScene> OutScene)
 {
-	OutScene->SetHost(TEXT("C4DTranslator"));
-	
-#ifdef _CHECK_DYNAMIC_IMPORTER_
-	if (IDatasmithC4DDynamicImporterModule::IsAvailable())
-	{
-		FDatasmithC4DImportOptions C4DImportOptions;
-		C4DImportOptions.bImportEmptyMesh = ImportOptions->bImportEmptyMesh;
-		C4DImportOptions.bExportToUDatasmith = ImportOptions->bExportToUDatasmith;
-		C4DImportOptions.bAlwaysGenerateNormals = ImportOptions->bAlwaysGenerateNormals;
-		C4DImportOptions.ScaleVertices = ImportOptions->ScaleVertices;
-		C4DImportOptions.bOptimizeEmptySingleChildActors = ImportOptions->bOptimizeEmptySingleChildActors;
+	NotifyImprovedImporter();
 
-		IDatasmithC4DDynamicImporterModule& DynamicModule = IDatasmithC4DDynamicImporterModule::Get();
-		if (DynamicModule.TryLoadingCineware())
-		{
-			Importer = DynamicModule.GetDynamicImporter(OutScene, C4DImportOptions);
-		}
-	}
-#endif
+	OutScene->SetHost(TEXT("C4DTranslator"));
 
-	if(!Importer.IsValid())
+	if (!Importer.IsValid())
 	{
 		FDatasmithC4DImportOptions C4DImportOptions;
 		Importer = MakeShared<FDatasmithC4DImporter>(OutScene, C4DImportOptions);
@@ -133,4 +123,32 @@ TStrongObjectPtr<UDatasmithC4DImportOptions>& FDatasmithC4DTranslator::GetOrCrea
 	return ImportOptions;
 }
 
+constexpr auto IMPROVED_IMPORTER_URL = TEXT("https://www.maxon.net/unreal");
+
+void NotifyImprovedImporter()
+{	
+	if (IsRunningCommandlet())
+		return;
+
+	UE_LOG(DatasmithC4DTranslatorLog, Warning, TEXT("Old DatasmithC4DTranslator loaded"));
+
+	const FText Msg = FText::Format(LOCTEXT("DatasmithC4DImporterLoaded",
+		"Improved Cineware Import available at {0}"), FText::FromString(IMPROVED_IMPORTER_URL));
+
+	FNotificationInfo Info(Msg);
+	Info.ExpireDuration = 8.0f;
+	Info.bUseLargeFont = true;
+	Info.bUseSuccessFailIcons = true;
+	Info.Image = FAppStyle::GetBrush(TEXT("MessageLog.Warning"));
+	Info.Hyperlink = FSimpleDelegate::CreateLambda([]() {
+		const FString URL = IMPROVED_IMPORTER_URL;
+		FPlatformProcess::LaunchURL(*URL, nullptr, nullptr);
+		});
+	Info.HyperlinkText = LOCTEXT("GoToMaxon", "Go to Maxon...");
+
+	FSlateNotificationManager::Get().AddNotification(Info);
+}
+
+#undef LOCTEXT_NAMESPACE
+
 #endif // _MELANGE_SDK_
\ No newline at end of file
diff --git a/Engine/Plugins/Enterprise/DatasmithC4DImporter/Source/DatasmithC4DTranslator/Private/DatasmithC4DTranslator.h b/Engine/Plugins/Enterprise/DatasmithC4DImporter/Source/DatasmithC4DTranslator/Private/DatasmithC4DTranslator.h
--- a/Engine/Plugins/Enterprise/DatasmithC4DImporter/Source/DatasmithC4DTranslator/Private/DatasmithC4DTranslator.h
+++ b/Engine/Plugins/Enterprise/DatasmithC4DImporter/Source/DatasmithC4DTranslator/Private/DatasmithC4DTranslator.h
@@ -6,7 +6,6 @@
 #include "DatasmithC4DImportOptions.h"
 #include "DatasmithTranslator.h"
 #include "DatasmithC4DImporter.h"
-#include "IDatasmithC4DImporter.h"
 
 #include "CoreMinimal.h"
 
@@ -15,7 +14,7 @@ class FDatasmithC4DTranslator : public IDatasmithTranslator
 public:
 	virtual FName GetFName() const override { return "DatasmithC4DTranslator"; };
 
-#if !defined(_MELANGE_SDK_) && !defined(_CINEWARE_SDK_)
+#if !defined(_MELANGE_SDK_)
 	virtual void Initialize(FDatasmithTranslatorCapabilities& OutCapabilities) override { OutCapabilities.bIsEnabled = false; }
 #else
 
@@ -39,6 +38,6 @@ private:
 private:
 	/** -1: try load dynamic initially; 0: dynamic cannot be loaded; 1: dynamic successful loaded */
 	TStrongObjectPtr<UDatasmithC4DImportOptions> ImportOptions;
-	TSharedPtr<class IDatasmithC4DImporter> Importer;
+	TSharedPtr<FDatasmithC4DImporter> Importer;
 #endif
 };
diff --git a/Engine/Plugins/Enterprise/DatasmithC4DImporter/Source/DatasmithC4DTranslator/Private/DatasmithC4DTranslatorModule.cpp b/Engine/Plugins/Enterprise/DatasmithC4DImporter/Source/DatasmithC4DTranslator/Private/DatasmithC4DTranslatorModule.cpp
--- a/Engine/Plugins/Enterprise/DatasmithC4DImporter/Source/DatasmithC4DTranslator/Private/DatasmithC4DTranslatorModule.cpp
+++ b/Engine/Plugins/Enterprise/DatasmithC4DImporter/Source/DatasmithC4DTranslator/Private/DatasmithC4DTranslatorModule.cpp
@@ -4,7 +4,6 @@
 
 #ifdef _MELANGE_SDK_
 
-#include "IDatasmithC4DImporter.h"
 #include "DatasmithC4DTranslator.h"
 #include "DatasmithC4DImporterMaterialSelector.h"
 
diff --git a/Engine/Plugins/Enterprise/DatasmithC4DImporter/Source/DatasmithC4DTranslator/Public/IDatasmithC4DImporter.h b/Engine/Plugins/Enterprise/DatasmithC4DImporter/Source/DatasmithC4DTranslator/Public/IDatasmithC4DImporter.h
deleted file mode 100644
