[Datasmith] Allow C4D translator to be used at runtime
Note: remaining banned words will be removed on a subsequent submission addressing the root issue in the Datasmith API.

diff --git a/Engine/Plugins/Enterprise/DatasmithC4DImporter/DatasmithC4DImporter.uplugin b/Engine/Plugins/Enterprise/DatasmithC4DImporter/DatasmithC4DImporter.uplugin
--- a/Engine/Plugins/Enterprise/DatasmithC4DImporter/DatasmithC4DImporter.uplugin
+++ b/Engine/Plugins/Enterprise/DatasmithC4DImporter/DatasmithC4DImporter.uplugin
@@ -23,7 +23,7 @@
 	[
 		{
 			"Name" : "DatasmithC4DTranslator",
-			"Type" : "Editor",
+			"Type" : "Runtime",
 			"LoadingPhase" : "PostDefault",
 			"PlatformAllowList" :
 			[
diff --git a/Engine/Plugins/Enterprise/DatasmithC4DImporter/Source/DatasmithC4DTranslator/Private/DatasmithC4DImporterMaterialSelector.cpp b/Engine/Plugins/Enterprise/DatasmithC4DImporter/Source/DatasmithC4DTranslator/Private/DatasmithC4DImporterMaterialSelector.cpp
--- a/Engine/Plugins/Enterprise/DatasmithC4DImporter/Source/DatasmithC4DTranslator/Private/DatasmithC4DImporterMaterialSelector.cpp
+++ b/Engine/Plugins/Enterprise/DatasmithC4DImporter/Source/DatasmithC4DTranslator/Private/DatasmithC4DImporterMaterialSelector.cpp
@@ -13,17 +13,17 @@
 
 FDatasmithC4DImporterMaterialSelector::FDatasmithC4DImporterMaterialSelector()
 {
-	MasterMaterial.FromSoftObjectPath( FSoftObjectPath("/DatasmithContent/Materials/C4DMaster.C4DMaster") );
+	BaseMaterial.FromSoftObjectPath( FSoftObjectPath("/DatasmithContent/Materials/C4DMaster.C4DMaster") );
 }
 
 bool FDatasmithC4DImporterMaterialSelector::IsValid() const
 {
-	return MasterMaterial.IsValid();
+	return BaseMaterial.IsValid();
 }
 
 const FDatasmithMasterMaterial& FDatasmithC4DImporterMaterialSelector::GetMasterMaterial( const TSharedPtr< IDatasmithMasterMaterialElement >& /*InDatasmithMaterial*/ ) const
 {
-	return MasterMaterial;
+	return BaseMaterial;
 }
 
 void FDatasmithC4DImporterMaterialSelector::FinalizeMaterialInstance(const TSharedPtr<IDatasmithMasterMaterialElement>& InDatasmithMaterial, UMaterialInstanceConstant * MaterialInstance) const
diff --git a/Engine/Plugins/Enterprise/DatasmithC4DImporter/Source/DatasmithC4DTranslator/Private/DatasmithC4DImporterMaterialSelector.h b/Engine/Plugins/Enterprise/DatasmithC4DImporter/Source/DatasmithC4DTranslator/Private/DatasmithC4DImporterMaterialSelector.h
--- a/Engine/Plugins/Enterprise/DatasmithC4DImporter/Source/DatasmithC4DTranslator/Private/DatasmithC4DImporterMaterialSelector.h
+++ b/Engine/Plugins/Enterprise/DatasmithC4DImporter/Source/DatasmithC4DTranslator/Private/DatasmithC4DImporterMaterialSelector.h
@@ -23,5 +23,5 @@ public:
 
 private:
 
-	FDatasmithMasterMaterial MasterMaterial;
+	FDatasmithMasterMaterial BaseMaterial;
 };
diff --git a/Engine/Plugins/Enterprise/DatasmithC4DImporter/Source/DatasmithC4DTranslator/Private/DatasmithC4DTranslator.cpp b/Engine/Plugins/Enterprise/DatasmithC4DImporter/Source/DatasmithC4DTranslator/Private/DatasmithC4DTranslator.cpp
--- a/Engine/Plugins/Enterprise/DatasmithC4DImporter/Source/DatasmithC4DTranslator/Private/DatasmithC4DTranslator.cpp
+++ b/Engine/Plugins/Enterprise/DatasmithC4DImporter/Source/DatasmithC4DTranslator/Private/DatasmithC4DTranslator.cpp
@@ -8,13 +8,16 @@
 #include "DatasmithSceneSource.h"
 #include "IDatasmithSceneElements.h"
 
-#include "Framework/Notifications/NotificationManager.h"
 #include "GenericPlatform/GenericPlatformProcess.h"
 #include "Misc/FileHelper.h"
 #include "Misc/Paths.h"
-#include "Styling/AppStyle.h"
 #include "Templates/TypeHash.h"
+
+#if WITH_EDITOR
+#include "Framework/Notifications/NotificationManager.h"
+#include "Styling/AppStyle.h"
 #include "Widgets/Notifications/SNotificationList.h"
+#endif
 
 #define LOCTEXT_NAMESPACE "DatasmithC4DTranslator"
 
@@ -105,7 +108,9 @@ void FDatasmithC4DTranslator::SetSceneImportOptions(TArray<TStrongObjectPtr<UDat
 	{
 		FDatasmithC4DImportOptions C4DImportOptions;
 		C4DImportOptions.bImportEmptyMesh = ImportOptions->bImportEmptyMesh;
+#if WITH_EDITORONLY_DATA
 		C4DImportOptions.bExportToUDatasmith = ImportOptions->bExportToUDatasmith;
+#endif
 		C4DImportOptions.bAlwaysGenerateNormals = ImportOptions->bAlwaysGenerateNormals;
 		C4DImportOptions.ScaleVertices = ImportOptions->ScaleVertices;
 		C4DImportOptions.bOptimizeEmptySingleChildActors = ImportOptions->bOptimizeEmptySingleChildActors;
@@ -126,9 +131,17 @@ TStrongObjectPtr<UDatasmithC4DImportOptions>& FDatasmithC4DTranslator::GetOrCrea
 constexpr auto IMPROVED_IMPORTER_URL = TEXT("https://www.maxon.net/unreal");
 
 void NotifyImprovedImporter()
-{	
+{
+#if WITH_EDITOR
+	if (!IsInGameThread() && !IsInSlateThread() && !IsInAsyncLoadingThread())
+	{
+		return;
+	}
+
 	if (IsRunningCommandlet())
+	{
 		return;
+	}
 
 	UE_LOG(DatasmithC4DTranslatorLog, Warning, TEXT("Old DatasmithC4DTranslator loaded"));
 
@@ -147,6 +160,7 @@ void NotifyImprovedImporter()
 	Info.HyperlinkText = LOCTEXT("GoToMaxon", "Go to Maxon...");
 
 	FSlateNotificationManager::Get().AddNotification(Info);
+#endif
 }
 
 #undef LOCTEXT_NAMESPACE
