[Non-inclusivity] Removed banned words from Datasmith API and code using the API
 - IDatasmithMasterMaterialElement has become IDatasmithMaterialInstanceElement
 - EDatasmithMasterMateriallType has become EDatasmithReferenceMaterialType
 - IDatasmithMasterMaterialElement is signaled as deprecated and not supported in 5.2
 - The MasterMaterial tag is replaced by MaterialInstance, reading and writing. Reading the MasterMaterial tag is supported for previous versions.
 - Clients should still be able to compile against Datasmith SDK without changing anything. However, warnings of deprecation will be logged.
 - Directories, files and assets which named contained  banned words have been renamed too. Redirect directives in the DatasmithContent module have been added for existing assets

diff --git a/Engine/Plugins/Enterprise/DatasmithC4DImporter/Source/DatasmithC4DTranslator/Private/DatasmithC4DImporter.cpp b/Engine/Plugins/Enterprise/DatasmithC4DImporter/Source/DatasmithC4DTranslator/Private/DatasmithC4DImporter.cpp
--- a/Engine/Plugins/Enterprise/DatasmithC4DImporter/Source/DatasmithC4DTranslator/Private/DatasmithC4DImporter.cpp
+++ b/Engine/Plugins/Enterprise/DatasmithC4DImporter/Source/DatasmithC4DTranslator/Private/DatasmithC4DImporter.cpp
@@ -945,7 +945,7 @@ namespace
 		return Result;
 	}
 
-	void AddColorToMaterial(const TSharedPtr<IDatasmithMasterMaterialElement>& Material, const FString& DatasmithPropName, const FLinearColor& LinearColor)
+	void AddColorToMaterial(const TSharedPtr<IDatasmithMaterialInstanceElement>& Material, const FString& DatasmithPropName, const FLinearColor& LinearColor)
 	{
 		TSharedPtr<IDatasmithKeyValueProperty> MaterialPropertyPtr = FDatasmithSceneFactory::CreateKeyValueProperty(*DatasmithPropName);
 		MaterialPropertyPtr->SetPropertyType(EDatasmithKeyValuePropertyType::Color);
@@ -954,7 +954,7 @@ namespace
 		Material->AddProperty(MaterialPropertyPtr);
 	}
 
-	void AddFloatToMaterial(const TSharedPtr<IDatasmithMasterMaterialElement>& Material, const FString& DatasmithPropName, float Value)
+	void AddFloatToMaterial(const TSharedPtr<IDatasmithMaterialInstanceElement>& Material, const FString& DatasmithPropName, float Value)
 	{
 		TSharedPtr<IDatasmithKeyValueProperty> MaterialPropertyPtr = FDatasmithSceneFactory::CreateKeyValueProperty(*DatasmithPropName);
 		MaterialPropertyPtr->SetPropertyType(EDatasmithKeyValuePropertyType::Float);
@@ -963,7 +963,7 @@ namespace
 		Material->AddProperty(MaterialPropertyPtr);
 	}
 
-	void AddBoolToMaterial(const TSharedPtr<IDatasmithMasterMaterialElement>& Material, const FString& DatasmithPropName, bool bValue)
+	void AddBoolToMaterial(const TSharedPtr<IDatasmithMaterialInstanceElement>& Material, const FString& DatasmithPropName, bool bValue)
 	{
 		TSharedPtr<IDatasmithKeyValueProperty> MaterialPropertyPtr = FDatasmithSceneFactory::CreateKeyValueProperty(*DatasmithPropName);
 		MaterialPropertyPtr->SetPropertyType(EDatasmithKeyValuePropertyType::Bool);
@@ -972,7 +972,7 @@ namespace
 		Material->AddProperty(MaterialPropertyPtr);
 	}
 
-	void AddTextureToMaterial(const TSharedPtr<IDatasmithMasterMaterialElement>& Material, const FString& DatasmithPropName, TSharedPtr<IDatasmithTextureElement> Texture)
+	void AddTextureToMaterial(const TSharedPtr<IDatasmithMaterialInstanceElement>& Material, const FString& DatasmithPropName, TSharedPtr<IDatasmithTextureElement> Texture)
 	{
 		if (Texture == nullptr)
 		{
@@ -1255,18 +1255,18 @@ FString FDatasmithC4DImporter::GetBaseShaderTextureFilePath(melange::BaseList2D*
 	return TextureFilePath;
 }
 
-TSharedPtr<IDatasmithMasterMaterialElement> FDatasmithC4DImporter::ImportMaterial(melange::Material* InC4DMaterialPtr)
+TSharedPtr<IDatasmithMaterialInstanceElement> FDatasmithC4DImporter::ImportMaterial(melange::Material* InC4DMaterialPtr)
 {
 	TOptional<FString> DatasmithName = GetMelangeBaseList2dID(InC4DMaterialPtr);
 	if (!DatasmithName)
 	{
-		return TSharedPtr<IDatasmithMasterMaterialElement>();
+		return TSharedPtr<IDatasmithMaterialInstanceElement>();
 	}
 	FString DatasmithLabel = FDatasmithUtils::SanitizeObjectName(MelangeObjectName(InC4DMaterialPtr));
 
-	TSharedPtr<IDatasmithMasterMaterialElement> MaterialPtr = FDatasmithSceneFactory::CreateMasterMaterial(*(DatasmithName.GetValue()));
+	TSharedPtr<IDatasmithMaterialInstanceElement> MaterialPtr = FDatasmithSceneFactory::CreateMaterialInstance(*(DatasmithName.GetValue()));
 	MaterialPtr->SetLabel(*DatasmithLabel);
-	MaterialPtr->SetMaterialType(EDatasmithMasterMaterialType::Opaque);
+	MaterialPtr->SetMaterialType(EDatasmithReferenceMaterialType::Opaque);
 
 	// Color
 	bool bUseColor = InC4DMaterialPtr->GetChannelState(CHANNEL_COLOR);
@@ -1340,7 +1340,7 @@ TSharedPtr<IDatasmithMasterMaterialElement> FDatasmithC4DImporter::ImportMateria
 	AddBoolToMaterial(MaterialPtr, TEXT("Use_Transparency"), bUseTransparency);
 	if (bUseTransparency)
 	{
-		MaterialPtr->SetMaterialType(EDatasmithMasterMaterialType::Transparent);
+		MaterialPtr->SetMaterialType(EDatasmithReferenceMaterialType::Transparent);
 
 		melange::BaseList2D* TransparencyShader = MelangeGetLink(InC4DMaterialPtr, melange::MATERIAL_TRANSPARENCY_SHADER);
 		FString TransparencyMapPath = GetBaseShaderTextureFilePath(TransparencyShader);
@@ -1530,7 +1530,7 @@ TSharedPtr<IDatasmithMasterMaterialElement> FDatasmithC4DImporter::ImportMateria
 	AddBoolToMaterial(MaterialPtr, TEXT("Use_Alpha"), bUseAlpha);
 	if (bUseAlpha)
 	{
-		MaterialPtr->SetMaterialType(EDatasmithMasterMaterialType::CutOut);
+		MaterialPtr->SetMaterialType(EDatasmithReferenceMaterialType::CutOut);
 
 		melange::BaseList2D* AlphaShader = MelangeGetLink(InC4DMaterialPtr, melange::MATERIAL_ALPHA_SHADER);
 		FString AlphaMapPath = GetBaseShaderTextureFilePath(AlphaShader);
@@ -1562,12 +1562,12 @@ TSharedPtr<IDatasmithMasterMaterialElement> FDatasmithC4DImporter::ImportMateria
 	return MaterialPtr;
 }
 
-TSharedPtr<IDatasmithMasterMaterialElement> FDatasmithC4DImporter::ImportSimpleColorMaterial(melange::BaseObject* Object, int32 UseColor)
+TSharedPtr<IDatasmithMaterialInstanceElement> FDatasmithC4DImporter::ImportSimpleColorMaterial(melange::BaseObject* Object, int32 UseColor)
 {
 	TOptional<FString> DatasmithName = GetMelangeBaseList2dID(Object);
 	if (!DatasmithName)
 	{
-		return TSharedPtr<IDatasmithMasterMaterialElement>();
+		return TSharedPtr<IDatasmithMaterialInstanceElement>();
 	}
 	FString DatasmithLabel = FDatasmithUtils::SanitizeObjectName(MelangeObjectName(Object) + TEXT("_DisplayColor"));
 
@@ -1590,15 +1590,15 @@ TSharedPtr<IDatasmithMasterMaterialElement> FDatasmithC4DImporter::ImportSimpleC
 
 	FString MaterialHash = TEXT("DisplayColor_") + LexToString(GetTypeHash(DisplayColor));
 
-	TSharedPtr<IDatasmithMasterMaterialElement>& Material = MaterialNameToMaterialElement.FindOrAdd(MaterialHash);
+	TSharedPtr<IDatasmithMaterialInstanceElement>& Material = MaterialNameToMaterialElement.FindOrAdd(MaterialHash);
 	if (Material.IsValid())
 	{
 		return Material;
 	}
 
-	Material = FDatasmithSceneFactory::CreateMasterMaterial(*(DatasmithName.GetValue()));
+	Material = FDatasmithSceneFactory::CreateMaterialInstance(*(DatasmithName.GetValue()));
 	Material->SetLabel(*DatasmithLabel);
-	Material->SetMaterialType(EDatasmithMasterMaterialType::Opaque);
+	Material->SetMaterialType(EDatasmithReferenceMaterialType::Opaque);
 
 	// Color
 	AddColorToMaterial(Material, TEXT("Color"), FLinearColor(DisplayColor));
@@ -1618,7 +1618,7 @@ bool FDatasmithC4DImporter::ImportMaterialHierarchy(melange::BaseMaterial* InC4D
 	{
 		if (InC4DMaterialPtr->GetType() == Mmaterial)
 		{
-			if (TSharedPtr<IDatasmithMasterMaterialElement> DatasmithMaterial = ImportMaterial(static_cast<melange::Material*>(InC4DMaterialPtr)))
+			if (TSharedPtr<IDatasmithMaterialInstanceElement> DatasmithMaterial = ImportMaterial(static_cast<melange::Material*>(InC4DMaterialPtr)))
 			{
 				MaterialNameToMaterialElement.Add(DatasmithMaterial->GetName(), DatasmithMaterial);
 			}
@@ -1651,10 +1651,10 @@ FString FDatasmithC4DImporter::CustomizeMaterial(const FString& InMaterialID, co
 
 		if (OffsetX != 0.0F || OffsetY != 0.0F || TilesX != 1.0F || TilesY != 1.0F)
 		{
-			TSharedPtr<IDatasmithMasterMaterialElement> CustomizedMaterial = FDatasmithSceneFactory::CreateMasterMaterial(*CustomMaterialID);
+			TSharedPtr<IDatasmithMaterialInstanceElement> CustomizedMaterial = FDatasmithSceneFactory::CreateMaterialInstance(*CustomMaterialID);
 
 			// Create a copy of the original material
-			TSharedPtr<IDatasmithMasterMaterialElement> OriginalMaterial = MaterialNameToMaterialElement[InMaterialID];
+			TSharedPtr<IDatasmithMaterialInstanceElement> OriginalMaterial = MaterialNameToMaterialElement[InMaterialID];
 			for (int32 PropertyIndex = 0; PropertyIndex < OriginalMaterial->GetPropertiesCount(); ++PropertyIndex)
 			{
 				CustomizedMaterial->AddProperty(OriginalMaterial->GetProperty(PropertyIndex));
@@ -1764,7 +1764,7 @@ TSharedPtr<IDatasmithMeshActorElement> FDatasmithC4DImporter::ImportPolygon(mela
 		FString MaterialName = Pair.Value;
 
 		// Pick whether we use the display color material or a texturetag material
-		TSharedPtr<IDatasmithMasterMaterialElement> TargetMaterial = nullptr;
+		TSharedPtr<IDatasmithMaterialInstanceElement> TargetMaterial = nullptr;
 		if (UseColor == melange::ID_BASEOBJECT_USECOLOR_ALWAYS || UseColor == melange::ID_BASEOBJECT_USECOLOR_LAYER)
 		{
 			TargetMaterial = ImportSimpleColorMaterial(PolyObject, UseColor);
@@ -1777,7 +1777,7 @@ TSharedPtr<IDatasmithMeshActorElement> FDatasmithC4DImporter::ImportPolygon(mela
 			}
 			else
 			{
-				TSharedPtr<IDatasmithMasterMaterialElement>* FoundMaterial = MaterialNameToMaterialElement.Find(MaterialName);
+				TSharedPtr<IDatasmithMaterialInstanceElement>* FoundMaterial = MaterialNameToMaterialElement.Find(MaterialName);
 				if (FoundMaterial && FoundMaterial->IsValid())
 				{
 					TargetMaterial = *FoundMaterial;
diff --git a/Engine/Plugins/Enterprise/DatasmithC4DImporter/Source/DatasmithC4DTranslator/Private/DatasmithC4DImporter.h b/Engine/Plugins/Enterprise/DatasmithC4DImporter/Source/DatasmithC4DTranslator/Private/DatasmithC4DImporter.h
--- a/Engine/Plugins/Enterprise/DatasmithC4DImporter/Source/DatasmithC4DTranslator/Private/DatasmithC4DImporter.h
+++ b/Engine/Plugins/Enterprise/DatasmithC4DImporter/Source/DatasmithC4DTranslator/Private/DatasmithC4DImporter.h
@@ -21,7 +21,7 @@ class IDatasmithActorElement;
 class IDatasmithCameraActorElement;
 class IDatasmithLevelSequenceElement;
 class IDatasmithLightActorElement;
-class IDatasmithMasterMaterialElement;
+class IDatasmithMaterialInstanceElement;
 class IDatasmithMeshActorElement;
 class IDatasmithMeshElement;
 class IDatasmithScene;
@@ -96,15 +96,15 @@ public:
 	TSharedPtr<IDatasmithLightActorElement> ImportLight(melange::BaseObject* LightObject, const FString& DatasmithName, const FString& DatasmithLabel);
 	TSharedPtr<IDatasmithCameraActorElement> ImportCamera(melange::BaseObject* CameraObject, const FString& DatasmithName, const FString& DatasmithLabel);
 	TSharedPtr<IDatasmithMeshActorElement> ImportPolygon(melange::PolygonObject* PolyObject, const FString& DatasmithName, const FString& DatasmithLabel, const TArray<melange::TextureTag*>& TextureTags);
-	TSharedPtr<IDatasmithMasterMaterialElement> ImportMaterial(melange::Material* InC4DMaterialPtr);
-	TSharedPtr<IDatasmithMasterMaterialElement> ImportSimpleColorMaterial(melange::BaseObject* Object, int32 UseColor);
+	TSharedPtr<IDatasmithMaterialInstanceElement> ImportMaterial(melange::Material* InC4DMaterialPtr);
+	TSharedPtr<IDatasmithMaterialInstanceElement> ImportSimpleColorMaterial(melange::BaseObject* Object, int32 UseColor);
 	TSharedPtr<IDatasmithTextureElement> ImportTexture(const FString& TexturePath, EDatasmithTextureMode TextureMode);
 	TSharedPtr<IDatasmithMeshElement> ImportMesh(melange::PolygonObject* PolyObject, const FString& DatasmithMeshName, const FString& DatasmithLabel);
 
 	/** Parses the spline and its cache into SplineCurves so that it can be used as paths for animation later */
 	void ImportSpline(melange::SplineObject* ActorObject);
 
-	/** Traverse the melange material hierarchy contained in the c4d file and import each into IDatasmithMasterMaterialElements */
+	/** Traverse the melange material hierarchy contained in the c4d file and import each into IDatasmithMaterialInstanceElements */
 	bool ImportMaterialHierarchy(melange::BaseMaterial* InC4DMaterialPtr);
 
 	/** Uses ActorElementToC4DObject to find the corresponding melange object for ActorElement and adds all of its animations to LevelSequence */
@@ -117,7 +117,7 @@ public:
 	FString GetBaseShaderTextureFilePath(melange::BaseList2D* BaseShader);
 
 	/**
-	 * Generates a new copy of the IDatasmithMasterMaterialElement with name InMaterialID and alter its properties to match values retrieved
+	 * Generates a new copy of the IDatasmithMaterialInstanceElement with name InMaterialID and alter its properties to match values retrieved
 	 * from InTextureTag, and adds the new material to the Datasmith scene. This is used because texture tags are closer to material
 	 * instances, and may have different "overrides" for each property
 	 */
@@ -175,7 +175,7 @@ private:
 	TMap<melange::SplineObject*, TArray<FRichCurve>> SplineCurves;
 
 	/** Storage of created materials used by CustomizeMaterial to create new "material instances" */
-	TMap<FString, TSharedPtr<IDatasmithMasterMaterialElement>> MaterialNameToMaterialElement;
+	TMap<FString, TSharedPtr<IDatasmithMaterialInstanceElement>> MaterialNameToMaterialElement;
 
 	/** Cache meshes by hash to promote reusing StaticMeshes */
 	TMap<FString, TSharedRef<IDatasmithMeshElement>> PolygonHashToMeshElement;
diff --git a/Engine/Plugins/Enterprise/DatasmithC4DImporter/Source/DatasmithC4DTranslator/Private/DatasmithC4DImporterMaterialSelector.cpp b/Engine/Plugins/Enterprise/DatasmithC4DImporter/Source/DatasmithC4DTranslator/Private/DatasmithC4DImporterMaterialSelector.cpp
--- a/Engine/Plugins/Enterprise/DatasmithC4DImporter/Source/DatasmithC4DTranslator/Private/DatasmithC4DImporterMaterialSelector.cpp
+++ b/Engine/Plugins/Enterprise/DatasmithC4DImporter/Source/DatasmithC4DTranslator/Private/DatasmithC4DImporterMaterialSelector.cpp
@@ -13,20 +13,20 @@
 
 FDatasmithC4DImporterMaterialSelector::FDatasmithC4DImporterMaterialSelector()
 {
-	BaseMaterial.FromSoftObjectPath( FSoftObjectPath("/DatasmithContent/Materials/C4DMaster.C4DMaster") );
+	ReferenceMaterial.FromSoftObjectPath( FSoftObjectPath("/DatasmithContent/Materials/C4DReference.C4DReference") );
 }
 
 bool FDatasmithC4DImporterMaterialSelector::IsValid() const
 {
-	return BaseMaterial.IsValid();
+	return ReferenceMaterial.IsValid();
 }
 
-const FDatasmithMasterMaterial& FDatasmithC4DImporterMaterialSelector::GetMasterMaterial( const TSharedPtr< IDatasmithMasterMaterialElement >& /*InDatasmithMaterial*/ ) const
+const FDatasmithReferenceMaterial& FDatasmithC4DImporterMaterialSelector::GetReferenceMaterial( const TSharedPtr< IDatasmithMaterialInstanceElement >& /*InDatasmithMaterial*/ ) const
 {
-	return BaseMaterial;
+	return ReferenceMaterial;
 }
 
-void FDatasmithC4DImporterMaterialSelector::FinalizeMaterialInstance(const TSharedPtr<IDatasmithMasterMaterialElement>& InDatasmithMaterial, UMaterialInstanceConstant * MaterialInstance) const
+void FDatasmithC4DImporterMaterialSelector::FinalizeMaterialInstance(const TSharedPtr<IDatasmithMaterialInstanceElement>& InDatasmithMaterial, UMaterialInstanceConstant * MaterialInstance) const
 {
 	if (!InDatasmithMaterial.IsValid() || MaterialInstance == nullptr)
 	{
@@ -34,13 +34,13 @@ void FDatasmithC4DImporterMaterialSelector::FinalizeMaterialInstance(const TShar
 	}
 
 	// Set blend mode to translucent if material requires transparency.
-	if (InDatasmithMaterial->GetMaterialType() == EDatasmithMasterMaterialType::Transparent)
+	if (InDatasmithMaterial->GetMaterialType() == EDatasmithReferenceMaterialType::Transparent)
 	{
 		MaterialInstance->BasePropertyOverrides.bOverride_BlendMode = true;
 		MaterialInstance->BasePropertyOverrides.BlendMode = BLEND_Translucent;
 	}
 	// Set blend mode to masked if material has cutouts.
-	else if (InDatasmithMaterial->GetMaterialType() == EDatasmithMasterMaterialType::CutOut)
+	else if (InDatasmithMaterial->GetMaterialType() == EDatasmithReferenceMaterialType::CutOut)
 	{
 		MaterialInstance->BasePropertyOverrides.bOverride_BlendMode = true;
 		MaterialInstance->BasePropertyOverrides.BlendMode = BLEND_Masked;
diff --git a/Engine/Plugins/Enterprise/DatasmithC4DImporter/Source/DatasmithC4DTranslator/Private/DatasmithC4DImporterMaterialSelector.h b/Engine/Plugins/Enterprise/DatasmithC4DImporter/Source/DatasmithC4DTranslator/Private/DatasmithC4DImporterMaterialSelector.h
--- a/Engine/Plugins/Enterprise/DatasmithC4DImporter/Source/DatasmithC4DTranslator/Private/DatasmithC4DImporterMaterialSelector.h
+++ b/Engine/Plugins/Enterprise/DatasmithC4DImporter/Source/DatasmithC4DTranslator/Private/DatasmithC4DImporterMaterialSelector.h
@@ -3,14 +3,14 @@
 #pragma once
 
 #include "DatasmithDefinitions.h"
-#include "MasterMaterials/DatasmithMasterMaterial.h"
-#include "MasterMaterials/DatasmithMasterMaterialSelector.h"
+#include "ReferenceMaterials/DatasmithReferenceMaterial.h"
+#include "ReferenceMaterials/DatasmithReferenceMaterialSelector.h"
 
 #include "Templates/SharedPointer.h"
 
-class IDatasmithMasterMaterialElement;
+class IDatasmithMaterialInstanceElement;
 
-class FDatasmithC4DImporterMaterialSelector : public FDatasmithMasterMaterialSelector
+class FDatasmithC4DImporterMaterialSelector : public FDatasmithReferenceMaterialSelector
 {
 public:
 
@@ -18,10 +18,10 @@ public:
 	virtual ~FDatasmithC4DImporterMaterialSelector() = default;
 
 	virtual bool IsValid() const override;
-	virtual const FDatasmithMasterMaterial& GetMasterMaterial( const TSharedPtr< IDatasmithMasterMaterialElement >& /*InDatasmithMaterial*/ ) const override;
-	virtual void FinalizeMaterialInstance(const TSharedPtr< IDatasmithMasterMaterialElement >& InDatasmithMaterial, UMaterialInstanceConstant* MaterialInstance) const override;
+	virtual const FDatasmithReferenceMaterial& GetReferenceMaterial( const TSharedPtr< IDatasmithMaterialInstanceElement >& /*InDatasmithMaterial*/ ) const override;
+	virtual void FinalizeMaterialInstance(const TSharedPtr< IDatasmithMaterialInstanceElement >& InDatasmithMaterial, UMaterialInstanceConstant* MaterialInstance) const override;
 
 private:
 
-	FDatasmithMasterMaterial BaseMaterial;
+	FDatasmithReferenceMaterial ReferenceMaterial;
 };
diff --git a/Engine/Plugins/Enterprise/DatasmithC4DImporter/Source/DatasmithC4DTranslator/Private/DatasmithC4DTranslatorModule.cpp b/Engine/Plugins/Enterprise/DatasmithC4DImporter/Source/DatasmithC4DTranslator/Private/DatasmithC4DTranslatorModule.cpp
--- a/Engine/Plugins/Enterprise/DatasmithC4DImporter/Source/DatasmithC4DTranslator/Private/DatasmithC4DTranslatorModule.cpp
+++ b/Engine/Plugins/Enterprise/DatasmithC4DImporter/Source/DatasmithC4DTranslator/Private/DatasmithC4DTranslatorModule.cpp
@@ -7,7 +7,7 @@
 #include "DatasmithC4DTranslator.h"
 #include "DatasmithC4DImporterMaterialSelector.h"
 
-#include "MasterMaterials/DatasmithMasterMaterialManager.h"
+#include "ReferenceMaterials/DatasmithReferenceMaterialManager.h"
 
 #include "HAL/FileManager.h"
 #include "HAL/PlatformProcess.h"
@@ -38,7 +38,7 @@ public:
 		// Make sure the DatasmithImporter module exists and has been initialized before adding FDatasmithC4DTranslator's material selector
 		FModuleManager::Get().LoadModule(TEXT("DatasmithTranslator"));
 
-		FDatasmithMasterMaterialManager::Get().RegisterSelector(TEXT("C4DTranslator"), MakeShared< FDatasmithC4DImporterMaterialSelector >());
+		FDatasmithReferenceMaterialManager::Get().RegisterSelector(TEXT("C4DTranslator"), MakeShared< FDatasmithC4DImporterMaterialSelector >());
 
 		FString EnvVariable = FPlatformMisc::GetEnvironmentVariable(TEXT("DATASMITHC4D_DEBUG"));
 		bDebugMode = !EnvVariable.IsEmpty();
diff --git a/Engine/Plugins/Enterprise/DatasmithContent/Config/BaseDatasmithContent.ini b/Engine/Plugins/Enterprise/DatasmithContent/Config/BaseDatasmithContent.ini
new file mode 100644
--- /dev/null
+++ b/Engine/Plugins/Enterprise/DatasmithContent/Config/BaseDatasmithContent.ini
@@ -0,0 +1,19 @@
+[CoreRedirects]
++PackageRedirects=(OldName="/DatasmithContent/Materials/AliasMaster.AliasMaster",NewName="/DatasmithContent/Materials/AliasReference.AliasReference")
++PackageRedirects=(OldName="/DatasmithContent/Materials/C4DMaster.C4DMaster",NewName="/DatasmithContent/Materials/C4DReference.C4DReference")
++PackageRedirects=(OldName="/DatasmithContent/Materials/C4DMasterUV.C4DMasterUV",NewName="/DatasmithContent/Materials/C4DReferenceUV.C4DReferenceUV")
++PackageRedirects=(OldName="/DatasmithContent/Materials/CE_BaseMaster.CE_BaseMaster",NewName="/DatasmithContent/Materials/CE_BaseReference.CE_BaseReference")
++PackageRedirects=(OldName="/DatasmithContent/Materials/CE_OpacityMaster.CE_OpacityMaster",NewName="/DatasmithContent/Materials/CE_OpacityReference.CE_OpacityReference")
++PackageRedirects=(OldName="/DatasmithContent/Materials/CE_OpacitySimpleMaster.CE_OpacitySimpleMaster",NewName="/DatasmithContent/Materials/CE_OpacitySimpleReference.CE_OpacitySimpleReference")
++PackageRedirects=(OldName="/DatasmithContent/Materials/CE_OpaqueMaster.CE_OpaqueMaster",NewName="/DatasmithContent/Materials/CE_OpaqueReference.CE_OpaqueReference")
++PackageRedirects=(OldName="/DatasmithContent/Materials/IFC_ColorMaster.IFC_ColorMaster",NewName="/DatasmithContent/Materials/IFC_ColorReference.IFC_ColorReference")
++PackageRedirects=(OldName="/DatasmithContent/Materials/RevitMaster.RevitMaster",NewName="/DatasmithContent/Materials/RevitReference.RevitReference")
++PackageRedirects=(OldName="/DatasmithContent/Materials/SketchUpMaster.SketchUpMaster",NewName="/DatasmithContent/Materials/SketchUpReference.SketchUpReference")
++PackageRedirects=(OldName="/DatasmithContent/Materials/FBXImporter/DeltaGenMaster.DeltaGenMaster",NewName="/DatasmithContent/Materials/FBXImporter/DeltaGenReference.DeltaGenReference")
++PackageRedirects=(OldName="/DatasmithContent/Materials/FBXImporter/DeltaGenMasterAbsorb.DeltaGenMasterAbsorb",NewName="/DatasmithContent/Materials/FBXImporter/DeltaGenReferenceAbsorb.DeltaGenReferenceAbsorb")
++PackageRedirects=(OldName="/DatasmithContent/Materials/FBXImporter/DeltaGenMasterClearCoat.DeltaGenMasterClearCoat",NewName="/DatasmithContent/Materials/FBXImporter/DeltaGenReferenceClearCoat.DeltaGenReferenceClearCoat")
++PackageRedirects=(OldName="/DatasmithContent/Materials/FBXImporter/DeltaGenMasterTransparent.DeltaGenMasterTransparent",NewName="/DatasmithContent/Materials/FBXImporter/DeltaGenReferenceTransparent.DeltaGenReferenceTransparent")
++PackageRedirects=(OldName="/DatasmithContent/Materials/FBXImporter/VREDMaster.VREDMaster",NewName="/DatasmithContent/Materials/FBXImporter/VREDReference.VREDReference")
++PackageRedirects=(OldName="/DatasmithContent/Materials/FBXImporter/VREDMasterAbsorb.VREDMasterAbsorb",NewName="/DatasmithContent/Materials/FBXImporter/VREDReferenceAbsorb.VREDReferenceAbsorb")
++PackageRedirects=(OldName="/DatasmithContent/Materials/FBXImporter/VREDMasterClearCoat.VREDMasterClearCoat",NewName="/DatasmithContent/Materials/FBXImporter/VREDReferenceClearCoat.VREDReferenceClearCoat")
++PackageRedirects=(OldName="/DatasmithContent/Materials/FBXImporter/VREDMasterTransparent.VREDMasterTransparent",NewName="/DatasmithContent/Materials/FBXImporter/VREDReferenceTransparent.VREDReferenceTransparent")
diff --git a/Engine/Plugins/Enterprise/DatasmithFBXImporter/Source/DatasmithDeltaGenTranslator/Private/DatasmithDeltaGenImporter.cpp b/Engine/Plugins/Enterprise/DatasmithFBXImporter/Source/DatasmithDeltaGenTranslator/Private/DatasmithDeltaGenImporter.cpp
--- a/Engine/Plugins/Enterprise/DatasmithFBXImporter/Source/DatasmithDeltaGenTranslator/Private/DatasmithDeltaGenImporter.cpp
+++ b/Engine/Plugins/Enterprise/DatasmithFBXImporter/Source/DatasmithDeltaGenTranslator/Private/DatasmithDeltaGenImporter.cpp
@@ -510,7 +510,7 @@ TSharedPtr<IDatasmithActorElement> FDatasmithDeltaGenImporter::ConvertNode(const
 	return ActorElement;
 }
 
-FORCEINLINE void AddBoolProperty(IDatasmithMasterMaterialElement* Element, const FString& PropertyName, bool Value)
+FORCEINLINE void AddBoolProperty(IDatasmithMaterialInstanceElement* Element, const FString& PropertyName, bool Value)
 {
 	TSharedPtr<IDatasmithKeyValueProperty> MaterialProperty = FDatasmithSceneFactory::CreateKeyValueProperty(*PropertyName);
 	MaterialProperty->SetPropertyType(EDatasmithKeyValuePropertyType::Bool);
@@ -518,7 +518,7 @@ FORCEINLINE void AddBoolProperty(IDatasmithMasterMaterialElement* Element, const
 	Element->AddProperty(MaterialProperty);
 }
 
-FORCEINLINE void AddColorProperty(IDatasmithMasterMaterialElement* Element, const FString& PropertyName, const FVector4& Value)
+FORCEINLINE void AddColorProperty(IDatasmithMaterialInstanceElement* Element, const FString& PropertyName, const FVector4& Value)
 {
 	TSharedPtr<IDatasmithKeyValueProperty> MaterialProperty = FDatasmithSceneFactory::CreateKeyValueProperty(*PropertyName);
 	MaterialProperty->SetPropertyType(EDatasmithKeyValuePropertyType::Color);
@@ -527,7 +527,7 @@ FORCEINLINE void AddColorProperty(IDatasmithMasterMaterialElement* Element, cons
 	Element->AddProperty(MaterialProperty);
 }
 
-FORCEINLINE void AddFloatProperty(IDatasmithMasterMaterialElement* Element, const FString& PropertyName, float Value)
+FORCEINLINE void AddFloatProperty(IDatasmithMaterialInstanceElement* Element, const FString& PropertyName, float Value)
 {
 	TSharedPtr<IDatasmithKeyValueProperty> MaterialProperty = FDatasmithSceneFactory::CreateKeyValueProperty(*PropertyName);
 	MaterialProperty->SetPropertyType(EDatasmithKeyValuePropertyType::Float);
@@ -535,7 +535,7 @@ FORCEINLINE void AddFloatProperty(IDatasmithMasterMaterialElement* Element, cons
 	Element->AddProperty(MaterialProperty);
 }
 
-FORCEINLINE void AddStringProperty(IDatasmithMasterMaterialElement* Element, const FString& PropertyName, const FString& Value)
+FORCEINLINE void AddStringProperty(IDatasmithMaterialInstanceElement* Element, const FString& PropertyName, const FString& Value)
 {
 	TSharedPtr<IDatasmithKeyValueProperty> MaterialProperty = FDatasmithSceneFactory::CreateKeyValueProperty(*PropertyName);
 	MaterialProperty->SetPropertyType(EDatasmithKeyValuePropertyType::String);
@@ -543,7 +543,7 @@ FORCEINLINE void AddStringProperty(IDatasmithMasterMaterialElement* Element, con
 	Element->AddProperty(MaterialProperty);
 }
 
-FORCEINLINE void AddTextureProperty(IDatasmithMasterMaterialElement* Element, const FString& PropertyName, const FString& Path)
+FORCEINLINE void AddTextureProperty(IDatasmithMaterialInstanceElement* Element, const FString& PropertyName, const FString& Path)
 {
 	TSharedPtr<IDatasmithKeyValueProperty> MaterialProperty = FDatasmithSceneFactory::CreateKeyValueProperty(*PropertyName);
 	MaterialProperty->SetPropertyType(EDatasmithKeyValuePropertyType::Texture);
@@ -551,7 +551,7 @@ FORCEINLINE void AddTextureProperty(IDatasmithMasterMaterialElement* Element, co
 	Element->AddProperty(MaterialProperty);
 }
 
-TSharedPtr<IDatasmithTextureElement> CreateTextureAndTextureProperties(IDatasmithMasterMaterialElement* Element, const FString& TextureName, const FDatasmithFBXSceneMaterial::FTextureParams& Tex, EShadowTextureMode ShadowTextureMode)
+TSharedPtr<IDatasmithTextureElement> CreateTextureAndTextureProperties(IDatasmithMaterialInstanceElement* Element, const FString& TextureName, const FDatasmithFBXSceneMaterial::FTextureParams& Tex, EShadowTextureMode ShadowTextureMode)
 {
 	using MapEntry = TPairInitializer<const FString&, const EDatasmithTextureMode&>;
 	const static TMap<FString, EDatasmithTextureMode> TextureModes
@@ -584,7 +584,7 @@ TSharedPtr<IDatasmithTextureElement> CreateTextureAndTextureProperties(IDatasmit
 	// Pack all "texture properties". These are really material properties, but we'll use these to help
 	// map the texture correctly. Datasmith will bind these values to the material instance on creation and it will
 	// do that by property name, so it is imperative that they are like below (e.g. diffuseTranslation, glossyRotate, etc).
-	// Check the master material graphs to find the matching properties that will be filled in by datasmith
+	// Check the reference material graphs to find the matching properties that will be filled in by datasmith
 	AddTextureProperty(Element, TexHandle + TEXT("Path"), Tex.Path);
 
 	if (TexHandle == TEXT("TexAO"))
@@ -662,10 +662,10 @@ TSharedPtr<IDatasmithBaseMaterialElement> FDatasmithDeltaGenImporter::ConvertMat
 		return *OldMaterial;
 	}
 
-	TSharedRef<IDatasmithMasterMaterialElement> MaterialElement = FDatasmithSceneFactory::CreateMasterMaterial(*Material->Name);
+	TSharedRef<IDatasmithMaterialInstanceElement> MaterialElement = FDatasmithSceneFactory::CreateMaterialInstance(*Material->Name);
 	ImportedMaterials.Add(Material, MaterialElement);
 
-	IDatasmithMasterMaterialElement* El = &MaterialElement.Get();
+	IDatasmithMaterialInstanceElement* El = &MaterialElement.Get();
 	AddStringProperty(El, TEXT("Type"), Material->Type);
 	AddBoolProperty(El, TEXT("ReflectionIsActive"), true);
 
diff --git a/Engine/Plugins/Enterprise/DatasmithFBXImporter/Source/DatasmithDeltaGenTranslator/Private/DatasmithDeltaGenImporterMaterialSelector.cpp b/Engine/Plugins/Enterprise/DatasmithFBXImporter/Source/DatasmithDeltaGenTranslator/Private/DatasmithDeltaGenImporterMaterialSelector.cpp
--- a/Engine/Plugins/Enterprise/DatasmithFBXImporter/Source/DatasmithDeltaGenTranslator/Private/DatasmithDeltaGenImporterMaterialSelector.cpp
+++ b/Engine/Plugins/Enterprise/DatasmithFBXImporter/Source/DatasmithDeltaGenTranslator/Private/DatasmithDeltaGenImporterMaterialSelector.cpp
@@ -9,19 +9,19 @@
 
 FDatasmithDeltaGenImporterMaterialSelector::FDatasmithDeltaGenImporterMaterialSelector()
 {
-	// Master
-	MasterMaterial.FromSoftObjectPath( FSoftObjectPath("/DatasmithContent/Materials/FBXImporter/DeltaGenMaster.DeltaGenMaster") );
+	// Opaque
+	ReferenceMaterial.FromSoftObjectPath( FSoftObjectPath("/DatasmithContent/Materials/FBXImporter/DeltaGenReference.DeltaGenReference") );
 
 	// Transparent
-	MasterMaterialTransparent.FromSoftObjectPath( FSoftObjectPath("/DatasmithContent/Materials/FBXImporter/DeltaGenMasterTransparent.DeltaGenMasterTransparent") );
+	ReferenceMaterialTransparent.FromSoftObjectPath( FSoftObjectPath("/DatasmithContent/Materials/FBXImporter/DeltaGenReferenceTransparent.DeltaGenReferenceTransparent") );
 }
 
 bool FDatasmithDeltaGenImporterMaterialSelector::IsValid() const
 {
-	return MasterMaterial.IsValid() && MasterMaterialTransparent.IsValid();
+	return ReferenceMaterial.IsValid() && ReferenceMaterialTransparent.IsValid();
 }
 
-const FDatasmithMasterMaterial& FDatasmithDeltaGenImporterMaterialSelector::GetMasterMaterial( const TSharedPtr< IDatasmithMasterMaterialElement >& InDatasmithMaterial ) const
+const FDatasmithReferenceMaterial& FDatasmithDeltaGenImporterMaterialSelector::GetReferenceMaterial( const TSharedPtr< IDatasmithMaterialInstanceElement >& InDatasmithMaterial ) const
 {
 	// Transparent material
 	const TSharedPtr< IDatasmithKeyValueProperty > OpacityProperty = InDatasmithMaterial->GetPropertyByName(TEXT("Opacity"));
@@ -32,13 +32,13 @@ const FDatasmithMasterMaterial& FDatasmithDeltaGenImporterMaterialSelector::GetM
 		GetFloat( OpacityProperty, OpacityValue ) &&
 		OpacityValue < 1.f )
 	{
-		return MasterMaterialTransparent;
+		return ReferenceMaterialTransparent;
 	}
 
-	return MasterMaterial;
+	return ReferenceMaterial;
 }
 
-bool FDatasmithDeltaGenImporterMaterialSelector::IsValidMaterialType( EDatasmithMasterMaterialType InType ) const
+bool FDatasmithDeltaGenImporterMaterialSelector::IsValidMaterialType( EDatasmithReferenceMaterialType InType ) const
 {
-	return InType == EDatasmithMasterMaterialType::Auto || InType == EDatasmithMasterMaterialType::Opaque || InType == EDatasmithMasterMaterialType::Transparent;
+	return InType == EDatasmithReferenceMaterialType::Auto || InType == EDatasmithReferenceMaterialType::Opaque || InType == EDatasmithReferenceMaterialType::Transparent;
 }
diff --git a/Engine/Plugins/Enterprise/DatasmithFBXImporter/Source/DatasmithDeltaGenTranslator/Private/DatasmithDeltaGenImporterMaterialSelector.h b/Engine/Plugins/Enterprise/DatasmithFBXImporter/Source/DatasmithDeltaGenTranslator/Private/DatasmithDeltaGenImporterMaterialSelector.h
--- a/Engine/Plugins/Enterprise/DatasmithFBXImporter/Source/DatasmithDeltaGenTranslator/Private/DatasmithDeltaGenImporterMaterialSelector.h
+++ b/Engine/Plugins/Enterprise/DatasmithFBXImporter/Source/DatasmithDeltaGenTranslator/Private/DatasmithDeltaGenImporterMaterialSelector.h
@@ -2,25 +2,25 @@
 #pragma once
 
 #include "DatasmithDefinitions.h"
-#include "MasterMaterials/DatasmithMasterMaterial.h"
-#include "MasterMaterials/DatasmithMasterMaterialSelector.h"
+#include "ReferenceMaterials/DatasmithReferenceMaterial.h"
+#include "ReferenceMaterials/DatasmithReferenceMaterialSelector.h"
 
 #include "Templates/SharedPointer.h"
 
-class IDatasmithMasterMaterialElement;
+class IDatasmithMaterialInstanceElement;
 
-class FDatasmithDeltaGenImporterMaterialSelector : public FDatasmithMasterMaterialSelector
+class FDatasmithDeltaGenImporterMaterialSelector : public FDatasmithReferenceMaterialSelector
 {
 public:
 	FDatasmithDeltaGenImporterMaterialSelector();
 
 	virtual bool IsValid() const override;
-	virtual const FDatasmithMasterMaterial& GetMasterMaterial( const TSharedPtr< IDatasmithMasterMaterialElement >& InDatasmithMaterial ) const override;
+	virtual const FDatasmithReferenceMaterial& GetReferenceMaterial( const TSharedPtr< IDatasmithMaterialInstanceElement >& InDatasmithMaterial ) const override;
 
 protected:
-	bool IsValidMaterialType( EDatasmithMasterMaterialType InType ) const;
+	bool IsValidMaterialType( EDatasmithReferenceMaterialType InType ) const;
 
 private:
-	FDatasmithMasterMaterial MasterMaterial;
-	FDatasmithMasterMaterial MasterMaterialTransparent;
+	FDatasmithReferenceMaterial ReferenceMaterial;
+	FDatasmithReferenceMaterial ReferenceMaterialTransparent;
 };
diff --git a/Engine/Plugins/Enterprise/DatasmithFBXImporter/Source/DatasmithDeltaGenTranslator/Private/DatasmithDeltaGenTranslatorModule.cpp b/Engine/Plugins/Enterprise/DatasmithFBXImporter/Source/DatasmithDeltaGenTranslator/Private/DatasmithDeltaGenTranslatorModule.cpp
--- a/Engine/Plugins/Enterprise/DatasmithFBXImporter/Source/DatasmithDeltaGenTranslator/Private/DatasmithDeltaGenTranslatorModule.cpp
+++ b/Engine/Plugins/Enterprise/DatasmithFBXImporter/Source/DatasmithDeltaGenTranslator/Private/DatasmithDeltaGenTranslatorModule.cpp
@@ -2,7 +2,7 @@
 
 #include "DatasmithDeltaGenTranslatorModule.h"
 #include "DatasmithDeltaGenTranslator.h"
-#include "MasterMaterials/DatasmithMasterMaterialManager.h"
+#include "ReferenceMaterials/DatasmithReferenceMaterialManager.h"
 #include "DatasmithDeltaGenImporterMaterialSelector.h"
 
 #include "CoreMinimal.h"
@@ -16,7 +16,7 @@ public:
 		// Make sure the DatasmithImporter module exists and has been initialized before adding FDatasmithDeltaGenTranslator's material selector
 		FModuleManager::Get().LoadModule(TEXT("DatasmithTranslator"));
 
-		FDatasmithMasterMaterialManager::Get().RegisterSelector(TEXT("Deltagen"), MakeShared< FDatasmithDeltaGenImporterMaterialSelector >());
+		FDatasmithReferenceMaterialManager::Get().RegisterSelector(TEXT("Deltagen"), MakeShared< FDatasmithDeltaGenImporterMaterialSelector >());
 
 		Datasmith::RegisterTranslator<FDatasmithDeltaGenTranslator>();
 	}
diff --git a/Engine/Plugins/Enterprise/DatasmithFBXImporter/Source/DatasmithVREDTranslator/Private/DatasmithVREDImporter.cpp b/Engine/Plugins/Enterprise/DatasmithFBXImporter/Source/DatasmithVREDTranslator/Private/DatasmithVREDImporter.cpp
--- a/Engine/Plugins/Enterprise/DatasmithFBXImporter/Source/DatasmithVREDTranslator/Private/DatasmithVREDImporter.cpp
+++ b/Engine/Plugins/Enterprise/DatasmithFBXImporter/Source/DatasmithVREDTranslator/Private/DatasmithVREDImporter.cpp
@@ -976,7 +976,7 @@ TSharedPtr<IDatasmithActorElement> FDatasmithVREDImporter::ConvertNode(const TSh
 	return ActorElement;
 }
 
-void AddBoolProperty(IDatasmithMasterMaterialElement* Element, const FString& PropertyName, bool Value)
+void AddBoolProperty(IDatasmithMaterialInstanceElement* Element, const FString& PropertyName, bool Value)
 {
 	TSharedPtr<IDatasmithKeyValueProperty> MaterialProperty = FDatasmithSceneFactory::CreateKeyValueProperty(*PropertyName);
 	MaterialProperty->SetPropertyType(EDatasmithKeyValuePropertyType::Bool);
@@ -984,7 +984,7 @@ void AddBoolProperty(IDatasmithMasterMaterialElement* Element, const FString& Pr
 	Element->AddProperty(MaterialProperty);
 }
 
-void AddColorProperty(IDatasmithMasterMaterialElement* Element, const FString& PropertyName, const FVector4& Value)
+void AddColorProperty(IDatasmithMaterialInstanceElement* Element, const FString& PropertyName, const FVector4& Value)
 {
 	if (Value.ContainsNaN())
 	{
@@ -999,7 +999,7 @@ void AddColorProperty(IDatasmithMasterMaterialElement* Element, const FString& P
 	Element->AddProperty(MaterialProperty);
 }
 
-void AddFloatProperty(IDatasmithMasterMaterialElement* Element, const FString& PropertyName, float Value)
+void AddFloatProperty(IDatasmithMaterialInstanceElement* Element, const FString& PropertyName, float Value)
 {
 	if (!FMath::IsFinite(Value))
 	{
@@ -1013,7 +1013,7 @@ void AddFloatProperty(IDatasmithMasterMaterialElement* Element, const FString& P
 	Element->AddProperty(MaterialProperty);
 }
 
-void AddStringProperty(IDatasmithMasterMaterialElement* Element, const FString& PropertyName, const FString& Value)
+void AddStringProperty(IDatasmithMaterialInstanceElement* Element, const FString& PropertyName, const FString& Value)
 {
 	TSharedPtr<IDatasmithKeyValueProperty> MaterialProperty = FDatasmithSceneFactory::CreateKeyValueProperty(*PropertyName);
 	MaterialProperty->SetPropertyType(EDatasmithKeyValuePropertyType::String);
@@ -1021,7 +1021,7 @@ void AddStringProperty(IDatasmithMasterMaterialElement* Element, const FString&
 	Element->AddProperty(MaterialProperty);
 }
 
-void AddTextureProperty(IDatasmithMasterMaterialElement* Element, const FString& PropertyName, const FString& Path)
+void AddTextureProperty(IDatasmithMaterialInstanceElement* Element, const FString& PropertyName, const FString& Path)
 {
 	TSharedPtr<IDatasmithKeyValueProperty> MaterialProperty = FDatasmithSceneFactory::CreateKeyValueProperty(*PropertyName);
 	MaterialProperty->SetPropertyType(EDatasmithKeyValuePropertyType::Texture);
@@ -1031,8 +1031,8 @@ void AddTextureProperty(IDatasmithMasterMaterialElement* Element, const FString&
 
 // Pack all "texture properties". These are really material properties, but we'll use these to help
 // map the texture correctly. Datasmith will bind these values to the material instance on creation and it will
-// do that by property name, so make sure the names match the ones on the master materials
-void AddTextureMappingProperties(IDatasmithMasterMaterialElement* Element, const FString& TexHandle, const FDatasmithFBXSceneMaterial::FTextureParams& Tex)
+// do that by property name, so make sure the names match the ones on the reference materials
+void AddTextureMappingProperties(IDatasmithMaterialInstanceElement* Element, const FString& TexHandle, const FDatasmithFBXSceneMaterial::FTextureParams& Tex)
 {
 	AddTextureProperty(Element, TexHandle + TEXT("Path"), Tex.Path);
 	AddBoolProperty(Element, TexHandle + TEXT("IsActive"), !Tex.Path.IsEmpty() && Tex.bEnabled);
@@ -1075,7 +1075,7 @@ void AddTextureMappingProperties(IDatasmithMasterMaterialElement* Element, const
 	AddFloatProperty(Element, TexHandle + TEXT("Rotate"), Tex.Rotate);
 }
 
-TSharedPtr<IDatasmithTextureElement> CreateTextureElement(IDatasmithMasterMaterialElement* Element, const FString& TexHandle, const FString& TexturePath)
+TSharedPtr<IDatasmithTextureElement> CreateTextureElement(IDatasmithMaterialInstanceElement* Element, const FString& TexHandle, const FString& TexturePath)
 {
 	using TextureModeEntry = TPairInitializer<const FString&, const EDatasmithTextureMode&>;
 	const static TMap<FString, EDatasmithTextureMode> TextureHandleToModes
@@ -1176,7 +1176,7 @@ TSharedPtr<IDatasmithBaseMaterialElement> FDatasmithVREDImporter::ConvertMateria
 		return *OldMaterial;
 	}
 
-	TSharedPtr<IDatasmithMasterMaterialElement> MaterialElement = FDatasmithSceneFactory::CreateMasterMaterial(*Material->Name);
+	TSharedPtr<IDatasmithMaterialInstanceElement> MaterialElement = FDatasmithSceneFactory::CreateMaterialInstance(*Material->Name);
 	ImportedMaterials.Add(Material, MaterialElement);
 
 	// Base colors used for Chrome and BrushedMetal materials
@@ -1206,7 +1206,7 @@ TSharedPtr<IDatasmithBaseMaterialElement> FDatasmithVREDImporter::ConvertMateria
 		MetalColorEntry(20, FVector4(0.560f, 0.570f, 0.580f, 1.000f))  // Iron
 	});
 
-	IDatasmithMasterMaterialElement* El = MaterialElement.Get();
+	IDatasmithMaterialInstanceElement* El = MaterialElement.Get();
 
 	if (Material->Type.Equals(TEXT("UGlassMaterial")))
 	{
diff --git a/Engine/Plugins/Enterprise/DatasmithFBXImporter/Source/DatasmithVREDTranslator/Private/DatasmithVREDImporterMaterialSelector.cpp b/Engine/Plugins/Enterprise/DatasmithFBXImporter/Source/DatasmithVREDTranslator/Private/DatasmithVREDImporterMaterialSelector.cpp
--- a/Engine/Plugins/Enterprise/DatasmithFBXImporter/Source/DatasmithVREDTranslator/Private/DatasmithVREDImporterMaterialSelector.cpp
+++ b/Engine/Plugins/Enterprise/DatasmithFBXImporter/Source/DatasmithVREDTranslator/Private/DatasmithVREDImporterMaterialSelector.cpp
@@ -13,37 +13,37 @@
 
 FDatasmithVREDImporterMaterialSelector::FDatasmithVREDImporterMaterialSelector()
 {
-	FDatasmithMasterMaterial& Phong = MasterMaterials.Add(TEXT("UPhongMaterial"));
+	FDatasmithReferenceMaterial& Phong = ReferenceMaterials.Add(TEXT("UPhongMaterial"));
 	Phong.FromSoftObjectPath(FSoftObjectPath(TEXT("/DatasmithContent/Materials/FBXImporter/VRED/Phong.Phong")));
 
-	FDatasmithMasterMaterial& Plastic = MasterMaterials.Add(TEXT("UPlasticMaterial"));
+	FDatasmithReferenceMaterial& Plastic = ReferenceMaterials.Add(TEXT("UPlasticMaterial"));
 	Plastic.FromSoftObjectPath(FSoftObjectPath(TEXT("/DatasmithContent/Materials/FBXImporter/VRED/Plastic.Plastic")));
 
-	FDatasmithMasterMaterial& Glass = MasterMaterials.Add(TEXT("UGlassMaterial"));
+	FDatasmithReferenceMaterial& Glass = ReferenceMaterials.Add(TEXT("UGlassMaterial"));
 	Glass.FromSoftObjectPath(FSoftObjectPath(TEXT("/DatasmithContent/Materials/FBXImporter/VRED/Glass.Glass")));
 
-	FDatasmithMasterMaterial& Chrome = MasterMaterials.Add(TEXT("UChromeMaterial"));
+	FDatasmithReferenceMaterial& Chrome = ReferenceMaterials.Add(TEXT("UChromeMaterial"));
 	Chrome.FromSoftObjectPath(FSoftObjectPath(TEXT("/DatasmithContent/Materials/FBXImporter/VRED/Chrome.Chrome")));
 
-	FDatasmithMasterMaterial& BrushedMetal = MasterMaterials.Add(TEXT("UBrushedMetalMaterial"));
+	FDatasmithReferenceMaterial& BrushedMetal = ReferenceMaterials.Add(TEXT("UBrushedMetalMaterial"));
 	BrushedMetal.FromSoftObjectPath(FSoftObjectPath(TEXT("/DatasmithContent/Materials/FBXImporter/VRED/BrushedMetal.BrushedMetal")));
 
-	FDatasmithMasterMaterial& UnicolorCarpaint = MasterMaterials.Add(TEXT("UUnicolorPaintMaterial"));
+	FDatasmithReferenceMaterial& UnicolorCarpaint = ReferenceMaterials.Add(TEXT("UUnicolorPaintMaterial"));
 	UnicolorCarpaint.FromSoftObjectPath(FSoftObjectPath(TEXT("/DatasmithContent/Materials/FBXImporter/VRED/UnicolorCarpaint.UnicolorCarpaint")));
 }
 
 bool FDatasmithVREDImporterMaterialSelector::IsValid() const
 {
-	for (const auto& Pair : MasterMaterials)
+	for (const auto& Pair : ReferenceMaterials)
 	{
-		const FDatasmithMasterMaterial& Mat = Pair.Value;
+		const FDatasmithReferenceMaterial& Mat = Pair.Value;
 		if (!Mat.IsValid())
 		{
 			return false;
 		}
 	}
 
-	if (!MasterMaterials.Contains(DEFAULT_MATERIAL_NAME))
+	if (!ReferenceMaterials.Contains(DEFAULT_MATERIAL_NAME))
 	{
 		return false;
 	}
@@ -51,22 +51,22 @@ bool FDatasmithVREDImporterMaterialSelector::IsValid() const
 	return true;
 }
 
-const FDatasmithMasterMaterial& FDatasmithVREDImporterMaterialSelector::GetMasterMaterial( const TSharedPtr< IDatasmithMasterMaterialElement >& InDatasmithMaterial ) const
+const FDatasmithReferenceMaterial& FDatasmithVREDImporterMaterialSelector::GetReferenceMaterial( const TSharedPtr< IDatasmithMaterialInstanceElement >& InDatasmithMaterial ) const
 {
 	const TSharedPtr< IDatasmithKeyValueProperty > TypeProperty = InDatasmithMaterial->GetPropertyByName(TEXT("Type"));
 	FString TypeValue;
 	if (TypeProperty.IsValid() && GetString(TypeProperty, TypeValue))
 	{
-		if (const FDatasmithMasterMaterial* FoundMat = MasterMaterials.Find(TypeValue))
+		if (const FDatasmithReferenceMaterial* FoundMat = ReferenceMaterials.Find(TypeValue))
 		{
 			return *FoundMat;
 		}
 	}
 
-	return MasterMaterials[DEFAULT_MATERIAL_NAME];
+	return ReferenceMaterials[DEFAULT_MATERIAL_NAME];
 }
 
-void FDatasmithVREDImporterMaterialSelector::FinalizeMaterialInstance(const TSharedPtr< IDatasmithMasterMaterialElement >& InDatasmithMaterial, UMaterialInstanceConstant* MaterialInstance) const
+void FDatasmithVREDImporterMaterialSelector::FinalizeMaterialInstance(const TSharedPtr< IDatasmithMaterialInstanceElement >& InDatasmithMaterial, UMaterialInstanceConstant* MaterialInstance) const
 {
 	const TSharedPtr<IDatasmithKeyValueProperty> OpacityProperty = InDatasmithMaterial->GetPropertyByName(TEXT("Opacity"));
 	const TSharedPtr<IDatasmithKeyValueProperty> TransparencyTextureProperty = InDatasmithMaterial->GetPropertyByName(TEXT("TexTransparencyIsActive"));
@@ -93,7 +93,7 @@ void FDatasmithVREDImporterMaterialSelector::FinalizeMaterialInstance(const TSha
 	}
 }
 
-bool FDatasmithVREDImporterMaterialSelector::IsValidMaterialType( EDatasmithMasterMaterialType InType ) const
+bool FDatasmithVREDImporterMaterialSelector::IsValidMaterialType( EDatasmithReferenceMaterialType InType ) const
 {
-	return InType == EDatasmithMasterMaterialType::Auto || InType == EDatasmithMasterMaterialType::Opaque || InType == EDatasmithMasterMaterialType::Transparent;
+	return InType == EDatasmithReferenceMaterialType::Auto || InType == EDatasmithReferenceMaterialType::Opaque || InType == EDatasmithReferenceMaterialType::Transparent;
 }
diff --git a/Engine/Plugins/Enterprise/DatasmithFBXImporter/Source/DatasmithVREDTranslator/Private/DatasmithVREDImporterMaterialSelector.h b/Engine/Plugins/Enterprise/DatasmithFBXImporter/Source/DatasmithVREDTranslator/Private/DatasmithVREDImporterMaterialSelector.h
--- a/Engine/Plugins/Enterprise/DatasmithFBXImporter/Source/DatasmithVREDTranslator/Private/DatasmithVREDImporterMaterialSelector.h
+++ b/Engine/Plugins/Enterprise/DatasmithFBXImporter/Source/DatasmithVREDTranslator/Private/DatasmithVREDImporterMaterialSelector.h
@@ -2,25 +2,25 @@
 #pragma once
 
 #include "DatasmithDefinitions.h"
-#include "MasterMaterials/DatasmithMasterMaterial.h"
-#include "MasterMaterials/DatasmithMasterMaterialSelector.h"
+#include "ReferenceMaterials/DatasmithReferenceMaterial.h"
+#include "ReferenceMaterials/DatasmithReferenceMaterialSelector.h"
 
 #include "Templates/SharedPointer.h"
 
-class IDatasmithMasterMaterialElement;
+class IDatasmithMaterialInstanceElement;
 
-class FDatasmithVREDImporterMaterialSelector : public FDatasmithMasterMaterialSelector
+class FDatasmithVREDImporterMaterialSelector : public FDatasmithReferenceMaterialSelector
 {
 public:
 	FDatasmithVREDImporterMaterialSelector();
 
 	virtual bool IsValid() const override;
-	virtual const FDatasmithMasterMaterial& GetMasterMaterial( const TSharedPtr< IDatasmithMasterMaterialElement >& InDatasmithMaterial ) const override;
-	virtual void FinalizeMaterialInstance(const TSharedPtr< IDatasmithMasterMaterialElement >& InDatasmithMaterial, UMaterialInstanceConstant* MaterialInstance) const override;
+	virtual const FDatasmithReferenceMaterial& GetReferenceMaterial( const TSharedPtr< IDatasmithMaterialInstanceElement >& InDatasmithMaterial ) const override;
+	virtual void FinalizeMaterialInstance(const TSharedPtr< IDatasmithMaterialInstanceElement >& InDatasmithMaterial, UMaterialInstanceConstant* MaterialInstance) const override;
 
 protected:
-	bool IsValidMaterialType( EDatasmithMasterMaterialType InType ) const;
+	bool IsValidMaterialType( EDatasmithReferenceMaterialType InType ) const;
 
 private:
-	TMap<FString, FDatasmithMasterMaterial> MasterMaterials;
+	TMap<FString, FDatasmithReferenceMaterial> ReferenceMaterials;
 };
diff --git a/Engine/Plugins/Enterprise/DatasmithFBXImporter/Source/DatasmithVREDTranslator/Private/DatasmithVREDTranslatorModule.cpp b/Engine/Plugins/Enterprise/DatasmithFBXImporter/Source/DatasmithVREDTranslator/Private/DatasmithVREDTranslatorModule.cpp
--- a/Engine/Plugins/Enterprise/DatasmithFBXImporter/Source/DatasmithVREDTranslator/Private/DatasmithVREDTranslatorModule.cpp
+++ b/Engine/Plugins/Enterprise/DatasmithFBXImporter/Source/DatasmithVREDTranslator/Private/DatasmithVREDTranslatorModule.cpp
@@ -4,7 +4,7 @@
 
 #include "DatasmithVREDImporterMaterialSelector.h"
 #include "DatasmithVREDTranslator.h"
-#include "MasterMaterials/DatasmithMasterMaterialManager.h"
+#include "ReferenceMaterials/DatasmithReferenceMaterialManager.h"
 
 #include "CoreMinimal.h"
 #include "DatasmithTranslator.h"
@@ -17,7 +17,7 @@ public:
 		// Make sure the DatasmithImporter module exists and has been initialized before adding FDatasmithVREDTranslator's material selector
 		FModuleManager::Get().LoadModule(TEXT("DatasmithImporter"));
 
-		FDatasmithMasterMaterialManager::Get().RegisterSelector(TEXT("VRED"), MakeShared< FDatasmithVREDImporterMaterialSelector >());
+		FDatasmithReferenceMaterialManager::Get().RegisterSelector(TEXT("VRED"), MakeShared< FDatasmithVREDImporterMaterialSelector >());
 
 		Datasmith::RegisterTranslator<FDatasmithVREDTranslator>();
 	}
diff --git a/Engine/Plugins/Enterprise/DatasmithImporter/Source/DatasmithImporter/Private/DatasmithConsumer.cpp b/Engine/Plugins/Enterprise/DatasmithImporter/Source/DatasmithImporter/Private/DatasmithConsumer.cpp
--- a/Engine/Plugins/Enterprise/DatasmithImporter/Source/DatasmithImporter/Private/DatasmithConsumer.cpp
+++ b/Engine/Plugins/Enterprise/DatasmithImporter/Source/DatasmithImporter/Private/DatasmithConsumer.cpp
@@ -771,7 +771,7 @@ bool UDatasmithConsumer::BuildContexts()
 	ImportContextPtr->AssetsContext.StaticMeshesImportPackage.Reset();
 	ImportContextPtr->AssetsContext.TexturesImportPackage.Reset();
 	ImportContextPtr->AssetsContext.MaterialsImportPackage.Reset();
-	ImportContextPtr->AssetsContext.MasterMaterialsImportPackage.Reset();
+	ImportContextPtr->AssetsContext.ReferenceMaterialsImportPackage.Reset();
 	ImportContextPtr->AssetsContext.MaterialFunctionsImportPackage.Reset();
 	ImportContextPtr->AssetsContext.LevelSequencesImportPackage.Reset();
 	ImportContextPtr->AssetsContext.LevelVariantSetsImportPackage.Reset();
diff --git a/Engine/Plugins/Enterprise/DatasmithImporter/Source/DatasmithImporter/Private/DatasmithFileProducer.cpp b/Engine/Plugins/Enterprise/DatasmithImporter/Source/DatasmithImporter/Private/DatasmithFileProducer.cpp
--- a/Engine/Plugins/Enterprise/DatasmithImporter/Source/DatasmithImporter/Private/DatasmithFileProducer.cpp
+++ b/Engine/Plugins/Enterprise/DatasmithImporter/Source/DatasmithImporter/Private/DatasmithFileProducer.cpp
@@ -509,8 +509,8 @@ void UDatasmithFileProducer::PreventNameCollision()
 	UPackage* MaterialsImportPackage = NewObject< UPackage >( nullptr, *FPaths::Combine( TransientFolderPath, TEXT("Materials") ), RF_Transient );
 	MaterialsImportPackage->FullyLoad();
 
-	UPackage* MasterMaterialsImportPackage = NewObject< UPackage >(nullptr, *FPaths::Combine(TransientFolderPath, TEXT("Materials/Master")), RF_Transient);
-	MasterMaterialsImportPackage->FullyLoad();
+	UPackage* ReferenceMaterialsImportPackage = NewObject< UPackage >(nullptr, *FPaths::Combine(TransientFolderPath, TEXT("Materials/Reference")), RF_Transient);
+	ReferenceMaterialsImportPackage->FullyLoad();
 
 	UPackage* LevelSequencesImportPackage = NewObject< UPackage >( nullptr, *FPaths::Combine( TransientFolderPath, TEXT("Animations") ), RF_Transient );
 	LevelSequencesImportPackage->FullyLoad();
@@ -614,7 +614,7 @@ void UDatasmithFileProducer::PreventNameCollision()
 
 						if ( ParentPath.StartsWith( MaterialInstancePath ) )
 						{
-							MoveAsset( MaterialParent, MasterMaterialsImportPackage, true );
+							MoveAsset( MaterialParent, ReferenceMaterialsImportPackage, true );
 							ParentMaterials.Add( MaterialParent );
 						}
 					}
diff --git a/Engine/Plugins/Enterprise/DatasmithImporter/Source/DatasmithImporter/Private/DatasmithImportContext.cpp b/Engine/Plugins/Enterprise/DatasmithImporter/Source/DatasmithImporter/Private/DatasmithImportContext.cpp
--- a/Engine/Plugins/Enterprise/DatasmithImporter/Source/DatasmithImporter/Private/DatasmithImportContext.cpp
+++ b/Engine/Plugins/Enterprise/DatasmithImporter/Source/DatasmithImporter/Private/DatasmithImportContext.cpp
@@ -659,10 +659,10 @@ void FDatasmithAssetsImportContext::ReInit(const FString& NewRootFolder)
 	MaterialsImportPackage.Reset( NewObject< UPackage >( nullptr, *FPaths::Combine( TransientFolderPath, TEXT("Materials") ), RF_Transient ) );
 	MaterialsImportPackage->FullyLoad();
 
-	MasterMaterialsImportPackage.Reset(NewObject< UPackage >(nullptr, *FPaths::Combine(TransientFolderPath, TEXT("Materials/Master")), RF_Transient));
-	MasterMaterialsImportPackage->FullyLoad();
+	ReferenceMaterialsImportPackage.Reset(NewObject< UPackage >(nullptr, *FPaths::Combine(TransientFolderPath, TEXT("Materials/References")), RF_Transient));
+	ReferenceMaterialsImportPackage->FullyLoad();
 
-	MaterialFunctionsImportPackage.Reset( NewObject< UPackage >(nullptr, *FPaths::Combine(TransientFolderPath, TEXT("Materials/Master/Functions")), RF_Transient));
+	MaterialFunctionsImportPackage.Reset( NewObject< UPackage >(nullptr, *FPaths::Combine(TransientFolderPath, TEXT("Materials/References/Functions")), RF_Transient));
 	MaterialFunctionsImportPackage->FullyLoad();
 
 	LevelSequencesImportPackage.Reset( NewObject< UPackage >( nullptr, *FPaths::Combine( TransientFolderPath, TEXT("Animations") ), RF_Transient ) );
diff --git a/Engine/Plugins/Enterprise/DatasmithImporter/Source/DatasmithImporter/Private/DatasmithImportFactory.cpp b/Engine/Plugins/Enterprise/DatasmithImporter/Source/DatasmithImporter/Private/DatasmithImportFactory.cpp
--- a/Engine/Plugins/Enterprise/DatasmithImporter/Source/DatasmithImporter/Private/DatasmithImportFactory.cpp
+++ b/Engine/Plugins/Enterprise/DatasmithImporter/Source/DatasmithImporter/Private/DatasmithImportFactory.cpp
@@ -808,7 +808,7 @@ EReimportResult::Type UDatasmithImportFactory::ReimportStaticMesh(UStaticMesh* M
 	ImportContext.AssetsContext.LevelSequencesImportPackage.Reset();
 	ImportContext.AssetsContext.LevelVariantSetsImportPackage.Reset();
 	ImportContext.AssetsContext.MaterialsImportPackage.Reset();
-	ImportContext.AssetsContext.MasterMaterialsImportPackage.Reset();
+	ImportContext.AssetsContext.ReferenceMaterialsImportPackage.Reset();
 	ImportContext.AssetsContext.MaterialFunctionsImportPackage.Reset();
 	ImportContext.AssetsContext.TexturesImportPackage.Reset();
 
diff --git a/Engine/Plugins/Enterprise/DatasmithImporter/Source/DatasmithImporter/Private/DatasmithImporter.cpp b/Engine/Plugins/Enterprise/DatasmithImporter/Source/DatasmithImporter/Private/DatasmithImporter.cpp
--- a/Engine/Plugins/Enterprise/DatasmithImporter/Source/DatasmithImporter/Private/DatasmithImporter.cpp
+++ b/Engine/Plugins/Enterprise/DatasmithImporter/Source/DatasmithImporter/Private/DatasmithImporter.cpp
@@ -934,7 +934,7 @@ UObject* FDatasmithImporter::FinalizeMaterial( UObject* SourceMaterial, const TC
 {
 	TRACE_CPUPROFILER_EVENT_SCOPE(FDatasmithImporter::FinalizeMaterial);
 
-	// Finalizing the master material might add a remapping for the instance parent property so make sure we have a remapping map available
+	// Finalizing the material instance might add a remapping for the instance parent property so make sure we have a remapping map available
 	TOptional< TMap< UObject*, UObject* > > ReferencesToRemapLocal;
 	if ( !ReferencesToRemap )
 	{
diff --git a/Engine/Plugins/Enterprise/DatasmithImporter/Source/DatasmithImporter/Private/DatasmithLevelSequenceImporter.cpp b/Engine/Plugins/Enterprise/DatasmithImporter/Source/DatasmithImporter/Private/DatasmithLevelSequenceImporter.cpp
--- a/Engine/Plugins/Enterprise/DatasmithImporter/Source/DatasmithImporter/Private/DatasmithLevelSequenceImporter.cpp
+++ b/Engine/Plugins/Enterprise/DatasmithImporter/Source/DatasmithImporter/Private/DatasmithLevelSequenceImporter.cpp
@@ -463,7 +463,7 @@ namespace DatasmithLevelSequenceImporterImpl
 		}
 	}
 
-	void PopulateSubsequenceTrack(const TSharedRef<IDatasmithSubsequenceAnimationElement>& SubsequenceAnimation, bool bShouldClearMasterTrack, FDatasmithImportContext& ImportContext, ULevelSequence* LevelSequence, FFrameNumber& MinFrameNumber, FFrameNumber& MaxFrameNumber)
+	void PopulateSubsequenceTrack(const TSharedRef<IDatasmithSubsequenceAnimationElement>& SubsequenceAnimation, bool bShouldClearReferenceTrack, FDatasmithImportContext& ImportContext, ULevelSequence* LevelSequence, FFrameNumber& MinFrameNumber, FFrameNumber& MaxFrameNumber)
 	{
 		// Find the TargetMovieSceneSequence that will be added as a subsequence
 		TSharedPtr<IDatasmithLevelSequenceElement> PinnedTargetSequence = SubsequenceAnimation->GetSubsequence().Pin();
@@ -493,7 +493,7 @@ namespace DatasmithLevelSequenceImporterImpl
 		{
 			SubTrack = MovieScene->AddMasterTrack<UMovieSceneSubTrack>();
 		}
-		else if (SubTrack && bShouldClearMasterTrack)
+		else if (SubTrack && bShouldClearReferenceTrack)
 		{
 			SubTrack->RemoveAllAnimationData();
 		}
@@ -531,8 +531,8 @@ namespace DatasmithLevelSequenceImporterImpl
 		FFrameNumber MaxFrameNumber;
 
 		// We might be reusing the LevelSequence asset, so we must reset it before we add our first subsequence.
-		// We will add further subsequences onto the same master track though, so we can't clear again
-		bool bShouldClearSubsequenceMasterTrack = true;
+		// We will add further subsequences onto the same reference track though, so we can't clear again
+		bool bShouldClearSubsequenceReferenceTrack = true;
 
 		const int32 NumAnimations = LevelSequenceElement->GetAnimationsCount();
 		for (int32 AnimIndex = 0; AnimIndex < NumAnimations; ++AnimIndex)
@@ -556,9 +556,9 @@ namespace DatasmithLevelSequenceImporterImpl
 			else if (AnimationElement->IsSubType(EDatasmithElementAnimationSubType::SubsequenceAnimation))
 			{
 				TSharedRef<IDatasmithSubsequenceAnimationElement> SubsequenceAnimation = StaticCastSharedRef<IDatasmithSubsequenceAnimationElement>(AnimationElement.ToSharedRef());
-				PopulateSubsequenceTrack(SubsequenceAnimation, bShouldClearSubsequenceMasterTrack, ImportContext, LevelSequence, MinFrameNumber, MaxFrameNumber);
+				PopulateSubsequenceTrack(SubsequenceAnimation, bShouldClearSubsequenceReferenceTrack, ImportContext, LevelSequence, MinFrameNumber, MaxFrameNumber);
 
-				bShouldClearSubsequenceMasterTrack = false;
+				bShouldClearSubsequenceReferenceTrack = false;
 			}
 		}
 
diff --git a/Engine/Plugins/Enterprise/DatasmithImporter/Source/DatasmithImporter/Private/DatasmithMaterialExpressions.cpp b/Engine/Plugins/Enterprise/DatasmithImporter/Source/DatasmithImporter/Private/DatasmithMaterialExpressions.cpp
--- a/Engine/Plugins/Enterprise/DatasmithImporter/Source/DatasmithImporter/Private/DatasmithMaterialExpressions.cpp
+++ b/Engine/Plugins/Enterprise/DatasmithImporter/Source/DatasmithImporter/Private/DatasmithMaterialExpressions.cpp
@@ -3208,7 +3208,7 @@ EBlendMode FDatasmithMaterialExpressions::GetUEPbrImportBlendMode(const TSharedP
 UMaterialInterface* FDatasmithMaterialExpressions::CreateUEPbrMaterial(UPackage* Package, const TSharedPtr< IDatasmithUEPbrMaterialElement >& MaterialElement, FDatasmithAssetsImportContext& AssetsContext,
 	UMaterial* ExistingMaterial, EObjectFlags ObjectFlags)
 {
-	FString MaterialName = GenerateUniqueMaterialName(MaterialElement->GetParentLabel(), Package, AssetsContext.MasterMaterialNameProvider);
+	FString MaterialName = GenerateUniqueMaterialName(MaterialElement->GetParentLabel(), Package, AssetsContext.MaterialInstanceNameProvider);
 
 	// Verify that the material could be created in final package
 	FText FailReason;
@@ -3344,7 +3344,7 @@ UMaterialInterface* FDatasmithMaterialExpressions::CreateUEPbrMaterialInstance(U
 				}
 				else if(UTexture2D* Default2DTexture = Cast<UTexture2D>(DefaultTexture))
 				{
-					//We are adding a standard texture to a master material that takes virtual textures, we'll need to convert the master's material texture to standard.
+					//We are adding a standard texture to a reference material that takes virtual textures, we'll need to convert the reference's material texture to standard.
 					AssetsContext.VirtualTexturesToConvert.Add(Default2DTexture);
 				}
 			}
@@ -3429,7 +3429,7 @@ UMaterialInterface* FDatasmithMaterialExpressions::CreateUEPbrMaterialInstance(U
 								}
 								else if (UTexture2D* Default2DTexture = Cast<UTexture2D>(DefaultTexture))
 								{
-									//We are adding a standard texture to a master material that takes virtual textures, we'll need to convert the master's material texture to standard.
+									//We are adding a standard texture to a reference material that takes virtual textures, we'll need to convert the reference's material texture to standard.
 									AssetsContext.VirtualTexturesToConvert.Add(Default2DTexture);
 								}
 							}
diff --git a/Engine/Plugins/Enterprise/DatasmithImporter/Source/DatasmithImporter/Private/DatasmithMaterialImporter.cpp b/Engine/Plugins/Enterprise/DatasmithImporter/Source/DatasmithImporter/Private/DatasmithMaterialImporter.cpp
--- a/Engine/Plugins/Enterprise/DatasmithImporter/Source/DatasmithImporter/Private/DatasmithMaterialImporter.cpp
+++ b/Engine/Plugins/Enterprise/DatasmithImporter/Source/DatasmithImporter/Private/DatasmithMaterialImporter.cpp
@@ -6,9 +6,9 @@
 #include "DatasmithMaterialElements.h"
 #include "DatasmithMaterialExpressions.h"
 
-#include "MasterMaterials/DatasmithMasterMaterial.h"
-#include "MasterMaterials/DatasmithMasterMaterialManager.h"
-#include "MasterMaterials/DatasmithMasterMaterialSelector.h"
+#include "ReferenceMaterials/DatasmithReferenceMaterial.h"
+#include "ReferenceMaterials/DatasmithReferenceMaterialManager.h"
+#include "ReferenceMaterials/DatasmithReferenceMaterialSelector.h"
 #include "ObjectTemplates/DatasmithMaterialInstanceTemplate.h"
 #include "Utility/DatasmithImporterUtils.h"
 
@@ -213,10 +213,10 @@ UMaterialInterface* FDatasmithMaterialImporter::CreateMaterial( FDatasmithImport
 
 		Material = FDatasmithMaterialExpressions::CreateDatasmithMaterial(MaterialPackage, MaterialElement, ImportContext.AssetsContext, nullptr, ImportContext.ObjectFlags);
 	}
-	else if ( BaseMaterialElement->IsA( EDatasmithElementType::MasterMaterial ) )
+	else if ( BaseMaterialElement->IsA( EDatasmithElementType::MaterialInstance ) )
 	{
-		const TSharedRef< IDatasmithMasterMaterialElement >& MasterMaterialElement = StaticCastSharedRef< IDatasmithMasterMaterialElement >( BaseMaterialElement );
-		Material = ImportMasterMaterial( ImportContext, MasterMaterialElement, ExistingMaterial );
+		const TSharedRef< IDatasmithMaterialInstanceElement >& ReferenceMaterialElement = StaticCastSharedRef< IDatasmithMaterialInstanceElement >( BaseMaterialElement );
+		Material = ImportMaterialInstance( ImportContext, ReferenceMaterialElement, ExistingMaterial );
 	}
 	else if ( BaseMaterialElement->IsA( EDatasmithElementType::DecalMaterial ) )
 	{
@@ -236,7 +236,7 @@ UMaterialInterface* FDatasmithMaterialImporter::CreateMaterial( FDatasmithImport
 
 		if ( !ImportContext.ImportedParentMaterials.Contains( MaterialHash ) )
 		{
-			UMaterialInterface* ParentMaterial = FDatasmithMaterialExpressions::CreateUEPbrMaterial( ImportContext.AssetsContext.MasterMaterialsImportPackage.Get(), MaterialElement, ImportContext.AssetsContext, nullptr, ImportContext.ObjectFlags );
+			UMaterialInterface* ParentMaterial = FDatasmithMaterialExpressions::CreateUEPbrMaterial( ImportContext.AssetsContext.ReferenceMaterialsImportPackage.Get(), MaterialElement, ImportContext.AssetsContext, nullptr, ImportContext.ObjectFlags );
 
 			if (ParentMaterial == nullptr)
 			{
@@ -261,13 +261,13 @@ UMaterialInterface* FDatasmithMaterialImporter::CreateMaterial( FDatasmithImport
 	return Material;
 }
 
-UMaterialInterface* FDatasmithMaterialImporter::ImportMasterMaterial( FDatasmithImportContext& ImportContext, const TSharedRef< IDatasmithMasterMaterialElement >& MaterialElement, UMaterialInterface* ExistingMaterial )
+UMaterialInterface* FDatasmithMaterialImporter::ImportMaterialInstance( FDatasmithImportContext& ImportContext, const TSharedRef< IDatasmithMaterialInstanceElement >& MaterialElement, UMaterialInterface* ExistingMaterial )
 {
 	// Verify existing material is of the right class for further processing
 	UMaterialInstanceConstant* FoundConstantMaterial = Cast<UMaterialInstanceConstant>(ExistingMaterial);
 
-	FString Host = FDatasmithMasterMaterialManager::Get().GetHostFromString(ImportContext.Scene->GetHost());
-	TSharedPtr< FDatasmithMasterMaterialSelector > MaterialSelector = FDatasmithMasterMaterialManager::Get().GetSelector( *Host );
+	FString Host = FDatasmithReferenceMaterialManager::Get().GetHostFromString(ImportContext.Scene->GetHost());
+	TSharedPtr< FDatasmithReferenceMaterialSelector > MaterialSelector = FDatasmithReferenceMaterialManager::Get().GetSelector( *Host );
 
 	if (!MaterialSelector.IsValid())
 	{
@@ -276,25 +276,25 @@ UMaterialInterface* FDatasmithMaterialImporter::ImportMasterMaterial( FDatasmith
 		return nullptr;
 	}
 
-	const FDatasmithMasterMaterial* ParentMaterial = nullptr;
-	FDatasmithMasterMaterial CustomMasterMaterial; // MasterMaterial might point on this so keep them in the same scope
+	const FDatasmithReferenceMaterial* ParentMaterial = nullptr;
+	FDatasmithReferenceMaterial CustomReferenceMaterial; // ReferenceMaterial might point on this so keep them in the same scope
 
-	if ( MaterialElement->GetMaterialType() == EDatasmithMasterMaterialType::Custom )
+	if ( MaterialElement->GetMaterialType() == EDatasmithReferenceMaterialType::Custom )
 	{
-		CustomMasterMaterial.FromSoftObjectPath( FSoftObjectPath( MaterialElement->GetCustomMaterialPathName() ) );
-		if (!CustomMasterMaterial.IsValid())
+		CustomReferenceMaterial.FromSoftObjectPath( FSoftObjectPath( MaterialElement->GetCustomMaterialPathName() ) );
+		if (!CustomReferenceMaterial.IsValid())
 		{
-			ImportContext.LogError(FText::Format(LOCTEXT("NoMasterForPath", "No compatible asset for path '{0}'. Skipping material {1} ..."), FText::FromString(MaterialElement->GetCustomMaterialPathName()), FText::FromString(MaterialElement->GetName())));
+			ImportContext.LogError(FText::Format(LOCTEXT("NoReferenceForPath", "No compatible asset for path '{0}'. Skipping material {1} ..."), FText::FromString(MaterialElement->GetCustomMaterialPathName()), FText::FromString(MaterialElement->GetName())));
 			return nullptr;
 		}
 
-		ParentMaterial = &CustomMasterMaterial;
+		ParentMaterial = &CustomReferenceMaterial;
 	}
 	else
 	{
 		if ( MaterialSelector->IsValid() )
 		{
-			ParentMaterial = &MaterialSelector->GetMasterMaterial(MaterialElement);
+			ParentMaterial = &MaterialSelector->GetReferenceMaterial(MaterialElement);
 		}
 		else
 		{
@@ -340,7 +340,7 @@ UMaterialInterface* FDatasmithMaterialImporter::ImportMasterMaterial( FDatasmith
 
 		MaterialInstanceTemplate->ParentMaterial = MaterialInstance->Parent;
 
-		// Find matching master material parameters
+		// Find matching reference material parameters
 		for (int i = 0; i < MaterialElement->GetPropertiesCount(); ++i)
 		{
 			const TSharedPtr< IDatasmithKeyValueProperty > Property = MaterialElement->GetProperty(i);
diff --git a/Engine/Plugins/Enterprise/DatasmithImporter/Source/DatasmithImporter/Public/DatasmithImportContext.h b/Engine/Plugins/Enterprise/DatasmithImporter/Source/DatasmithImporter/Public/DatasmithImportContext.h
--- a/Engine/Plugins/Enterprise/DatasmithImporter/Source/DatasmithImporter/Public/DatasmithImportContext.h
+++ b/Engine/Plugins/Enterprise/DatasmithImporter/Source/DatasmithImporter/Public/DatasmithImportContext.h
@@ -174,8 +174,8 @@ struct DATASMITHIMPORTER_API FDatasmithAssetsImportContext
 	/** Package where the materials are stored until they are finalized */
 	TStrongObjectPtr< UPackage > MaterialsImportPackage;
 
-	/** Package where the master/parent materials are stored until they are finalized */
-	TStrongObjectPtr< UPackage > MasterMaterialsImportPackage;
+	/** Package where the reference/parent materials are stored until they are finalized */
+	TStrongObjectPtr< UPackage > ReferenceMaterialsImportPackage;
 
 	/** Package where the material functions are stored until they are finalized */
 	TStrongObjectPtr< UPackage > MaterialFunctionsImportPackage;
@@ -189,8 +189,8 @@ struct DATASMITHIMPORTER_API FDatasmithAssetsImportContext
 	/** Ensures uniqueness of the name of the materials being imported */
 	FDatasmithUniqueNameProvider MaterialNameProvider;
 
-	/** Ensures uniqueness of the name of the materials being imported in the master material package */
-	FDatasmithUniqueNameProvider MasterMaterialNameProvider;
+	/** Ensures uniqueness of the name of the materials being imported in the reference material package */
+	FDatasmithUniqueNameProvider MaterialInstanceNameProvider;
 
 	/** Ensures uniqueness of the name of the material functions being imported in the material function material package */
 	FDatasmithUniqueNameProvider MaterialFunctionNameProvider;
diff --git a/Engine/Plugins/Enterprise/DatasmithImporter/Source/DatasmithImporter/Public/DatasmithMaterialImporter.h b/Engine/Plugins/Enterprise/DatasmithImporter/Source/DatasmithImporter/Public/DatasmithMaterialImporter.h
--- a/Engine/Plugins/Enterprise/DatasmithImporter/Source/DatasmithImporter/Public/DatasmithMaterialImporter.h
+++ b/Engine/Plugins/Enterprise/DatasmithImporter/Source/DatasmithImporter/Public/DatasmithMaterialImporter.h
@@ -8,7 +8,7 @@
 class IDatasmithBaseMaterialElement;
 class IDatasmithDecalMaterialElement;
 class IDatasmithMaterialElement;
-class IDatasmithMasterMaterialElement;
+class IDatasmithMaterialInstanceElement;
 struct FDatasmithImportContext;
 class UMaterialInterface;
 class UMaterialFunction;
@@ -34,6 +34,6 @@ public:
 	static int32 GetMaterialRequirements(UMaterialInterface* MaterialInterface);
 
 private:
-	static UMaterialInterface* ImportMasterMaterial( FDatasmithImportContext& ImportContext, const TSharedRef< IDatasmithMasterMaterialElement >& MaterialElement, UMaterialInterface* ExistingMaterial );
+	static UMaterialInterface* ImportMaterialInstance( FDatasmithImportContext& ImportContext, const TSharedRef< IDatasmithMaterialInstanceElement >& MaterialElement, UMaterialInterface* ExistingMaterial );
 	static UMaterialInterface* ImportDecalMaterial( FDatasmithImportContext& ImportContext, const TSharedRef< IDatasmithDecalMaterialElement >& MaterialElement, UMaterialInterface* ExistingMaterial );
 };
\ No newline at end of file
diff --git a/Engine/Plugins/Enterprise/DatasmithImporter/Source/DatasmithTranslator/Private/DatasmithTranslatorModule.cpp b/Engine/Plugins/Enterprise/DatasmithImporter/Source/DatasmithTranslator/Private/DatasmithTranslatorModule.cpp
--- a/Engine/Plugins/Enterprise/DatasmithImporter/Source/DatasmithTranslator/Private/DatasmithTranslatorModule.cpp
+++ b/Engine/Plugins/Enterprise/DatasmithImporter/Source/DatasmithTranslator/Private/DatasmithTranslatorModule.cpp
@@ -2,34 +2,34 @@
 
 #include "DatasmithTranslatorModule.h"
 
-#include "MasterMaterials/DatasmithC4DMaterialSelector.h"
-#include "MasterMaterials/DatasmithCityEngineMaterialSelector.h"
-#include "MasterMaterials/DatasmithMasterMaterialManager.h"
-#include "MasterMaterials/DatasmithRevitMaterialSelector.h"
-#include "MasterMaterials/DatasmithSketchupMaterialSelector.h"
-#include "MasterMaterials/DatasmithStdMaterialSelector.h"
+#include "ReferenceMaterials/DatasmithC4DMaterialSelector.h"
+#include "ReferenceMaterials/DatasmithCityEngineMaterialSelector.h"
+#include "ReferenceMaterials/DatasmithReferenceMaterialManager.h"
+#include "ReferenceMaterials/DatasmithRevitMaterialSelector.h"
+#include "ReferenceMaterials/DatasmithSketchupMaterialSelector.h"
+#include "ReferenceMaterials/DatasmithStdMaterialSelector.h"
 
 void IDatasmithTranslatorModule::StartupModule()
 {
-	FDatasmithMasterMaterialManager::Create();
+	FDatasmithReferenceMaterialManager::Create();
 
-	//A minimal set of natively supported master materials
-	FDatasmithMasterMaterialManager::Get().RegisterSelector(TEXT("C4D"), MakeShared< FDatasmithC4DMaterialSelector >());
-	FDatasmithMasterMaterialManager::Get().RegisterSelector(TEXT("Revit"), MakeShared< FDatasmithRevitMaterialSelector >());
-	FDatasmithMasterMaterialManager::Get().RegisterSelector(TEXT("SketchUp"), MakeShared< FDatasmithSketchUpMaterialSelector >());
-	FDatasmithMasterMaterialManager::Get().RegisterSelector(TEXT("CityEngine"), MakeShared< FDatasmithCityEngineMaterialSelector >());
-	FDatasmithMasterMaterialManager::Get().RegisterSelector(TEXT("StdMaterial"), MakeShared< FDatasmithStdMaterialSelector >());
+	//A minimal set of natively supported reference materials
+	FDatasmithReferenceMaterialManager::Get().RegisterSelector(TEXT("C4D"), MakeShared< FDatasmithC4DMaterialSelector >());
+	FDatasmithReferenceMaterialManager::Get().RegisterSelector(TEXT("Revit"), MakeShared< FDatasmithRevitMaterialSelector >());
+	FDatasmithReferenceMaterialManager::Get().RegisterSelector(TEXT("SketchUp"), MakeShared< FDatasmithSketchUpMaterialSelector >());
+	FDatasmithReferenceMaterialManager::Get().RegisterSelector(TEXT("CityEngine"), MakeShared< FDatasmithCityEngineMaterialSelector >());
+	FDatasmithReferenceMaterialManager::Get().RegisterSelector(TEXT("StdMaterial"), MakeShared< FDatasmithStdMaterialSelector >());
 }
 
 void IDatasmithTranslatorModule::ShutdownModule()
 {
-	FDatasmithMasterMaterialManager::Get().UnregisterSelector(TEXT("C4D"));
-	FDatasmithMasterMaterialManager::Get().UnregisterSelector(TEXT("Revit"));
-	FDatasmithMasterMaterialManager::Get().UnregisterSelector(TEXT("SketchUp"));
-	FDatasmithMasterMaterialManager::Get().UnregisterSelector(TEXT("CityEngine"));
-	FDatasmithMasterMaterialManager::Get().UnregisterSelector(TEXT("StdMaterial"));
+	FDatasmithReferenceMaterialManager::Get().UnregisterSelector(TEXT("C4D"));
+	FDatasmithReferenceMaterialManager::Get().UnregisterSelector(TEXT("Revit"));
+	FDatasmithReferenceMaterialManager::Get().UnregisterSelector(TEXT("SketchUp"));
+	FDatasmithReferenceMaterialManager::Get().UnregisterSelector(TEXT("CityEngine"));
+	FDatasmithReferenceMaterialManager::Get().UnregisterSelector(TEXT("StdMaterial"));
 
-	FDatasmithMasterMaterialManager::Destroy();
+	FDatasmithReferenceMaterialManager::Destroy();
 }
 
 IMPLEMENT_MODULE(IDatasmithTranslatorModule, DatasmithTranslator);
\ No newline at end of file
diff --git a/Engine/Plugins/Experimental/Enterprise/DatasmithRuntime/Source/Private/DatasmithRuntimeModule.cpp b/Engine/Plugins/Experimental/Enterprise/DatasmithRuntime/Source/Private/DatasmithRuntimeModule.cpp
--- a/Engine/Plugins/Experimental/Enterprise/DatasmithRuntime/Source/Private/DatasmithRuntimeModule.cpp
+++ b/Engine/Plugins/Experimental/Enterprise/DatasmithRuntime/Source/Private/DatasmithRuntimeModule.cpp
@@ -10,7 +10,7 @@
 
 #include "DatasmithTranslatorModule.h"
 #include "DatasmithGLTFTranslatorModule.h"
-#include "MasterMaterials/DatasmithMasterMaterialManager.h"
+#include "ReferenceMaterials/DatasmithReferenceMaterialManager.h"
 
 #if WITH_EDITOR
 #include "Settings/ProjectPackagingSettings.h"
@@ -62,7 +62,7 @@ public:
 
 		DatasmithRuntime::FDestinationProxy::InitializeEndpointProxy();
 
-		FDatasmithMasterMaterialManager::Get().RegisterSelector(DatasmithRuntime::MATERIAL_HOST, MakeShared< FDatasmithRuntimeMaterialSelector >());
+		FDatasmithReferenceMaterialManager::Get().RegisterSelector(DatasmithRuntime::MATERIAL_HOST, MakeShared< FDatasmithRuntimeMaterialSelector >());
 
 		ADatasmithRuntimeActor::OnStartupModule();
 	}
@@ -71,7 +71,7 @@ public:
 	{
 		ADatasmithRuntimeActor::OnShutdownModule();
 		
-		FDatasmithMasterMaterialManager::Get().UnregisterSelector(DatasmithRuntime::MATERIAL_HOST);
+		FDatasmithReferenceMaterialManager::Get().UnregisterSelector(DatasmithRuntime::MATERIAL_HOST);
 
 		DatasmithRuntime::FDestinationProxy::ShutdownEndpointProxy();
 	}
diff --git a/Engine/Plugins/Experimental/Enterprise/DatasmithRuntime/Source/Private/DatasmithRuntimeUtils.h b/Engine/Plugins/Experimental/Enterprise/DatasmithRuntime/Source/Private/DatasmithRuntimeUtils.h
--- a/Engine/Plugins/Experimental/Enterprise/DatasmithRuntime/Source/Private/DatasmithRuntimeUtils.h
+++ b/Engine/Plugins/Experimental/Enterprise/DatasmithRuntime/Source/Private/DatasmithRuntimeUtils.h
@@ -17,7 +17,6 @@
 
 class IDatasmithBaseMaterialElement;
 class IDatasmithElement;
-class IDatasmithMasterMaterialElement;
 class IDatasmithMeshElement;
 class IDatasmithTextureElement;
 class IDatasmithUEPbrMaterialElement;
diff --git a/Engine/Plugins/Experimental/Enterprise/DatasmithRuntime/Source/Private/MaterialImportUtils.cpp b/Engine/Plugins/Experimental/Enterprise/DatasmithRuntime/Source/Private/MaterialImportUtils.cpp
--- a/Engine/Plugins/Experimental/Enterprise/DatasmithRuntime/Source/Private/MaterialImportUtils.cpp
+++ b/Engine/Plugins/Experimental/Enterprise/DatasmithRuntime/Source/Private/MaterialImportUtils.cpp
@@ -5,9 +5,9 @@
 #include "LogCategory.h"
 
 #include "DatasmithMaterialElements.h"
-#include "MasterMaterials/DatasmithMasterMaterial.h"
-#include "MasterMaterials/DatasmithMasterMaterialManager.h"
-#include "MasterMaterials/DatasmithMasterMaterialSelector.h"
+#include "ReferenceMaterials/DatasmithReferenceMaterial.h"
+#include "ReferenceMaterials/DatasmithReferenceMaterialManager.h"
+#include "ReferenceMaterials/DatasmithReferenceMaterialSelector.h"
 
 #include "Engine/Texture2D.h"
 #include "Materials/Material.h"
@@ -125,40 +125,40 @@ namespace DatasmithRuntime
 		return ParametersRef;
 	}
 
-	int32 ProcessMaterialElement(TSharedPtr< IDatasmithMasterMaterialElement > MasterMaterialElement, FTextureCallback TextureCallback)
+	int32 ProcessMaterialElement(TSharedPtr< IDatasmithMaterialInstanceElement > ReferenceMaterialElement, FTextureCallback TextureCallback)
 	{
-		TRACE_CPUPROFILER_EVENT_SCOPE(DatasmithRuntime::ProcessMasterMaterialElement);
+		TRACE_CPUPROFILER_EVENT_SCOPE(DatasmithRuntime::ProcessMaterialElement);
 
 		// Must be updated if FDatasmithMaterialImporter::GetMaterialRequirements changes
 		int32 MaterialRequirement = EMaterialRequirements::RequiresNormals | EMaterialRequirements::RequiresTangents;
 
-		if (!MasterMaterialElement.IsValid())
+		if (!ReferenceMaterialElement.IsValid())
 		{
 			return MaterialRequirement;
 		}
 
-		TSharedPtr< FDatasmithMasterMaterialSelector > MaterialSelector = FDatasmithMasterMaterialManager::Get().GetSelector(MATERIAL_HOST);
+		TSharedPtr< FDatasmithReferenceMaterialSelector > MaterialSelector = FDatasmithReferenceMaterialManager::Get().GetSelector(MATERIAL_HOST);
 
 		UMaterialInterface* Material = nullptr;
 
-		if (MasterMaterialElement->GetMaterialType() == EDatasmithMasterMaterialType::Custom)
+		if (ReferenceMaterialElement->GetMaterialType() == EDatasmithReferenceMaterialType::Custom)
 		{
-			FDatasmithMasterMaterial CustomMasterMaterial; // MasterMaterial might point on this so keep them in the same scope
+			FDatasmithReferenceMaterial CustomReferenceMaterial; // ReferenceMaterial might point on this so keep them in the same scope
 
-			CustomMasterMaterial.FromSoftObjectPath(FSoftObjectPath(MasterMaterialElement->GetCustomMaterialPathName()));
+			CustomReferenceMaterial.FromSoftObjectPath(FSoftObjectPath(ReferenceMaterialElement->GetCustomMaterialPathName()));
 
-			if (CustomMasterMaterial.IsValid())
+			if (CustomReferenceMaterial.IsValid())
 			{
-				Material = CustomMasterMaterial.GetMaterial();
+				Material = CustomReferenceMaterial.GetMaterial();
 			}
 		}
 		else if (MaterialSelector.IsValid() && MaterialSelector->IsValid())
 		{
-			const FDatasmithMasterMaterial& MasterMaterial = MaterialSelector->GetMasterMaterial(MasterMaterialElement);
+			const FDatasmithReferenceMaterial& ReferenceMaterial = MaterialSelector->GetReferenceMaterial(ReferenceMaterialElement);
 
-			if (MasterMaterial.IsValid())
+			if (ReferenceMaterial.IsValid())
 			{
-				Material = MasterMaterial.GetMaterial();
+				Material = ReferenceMaterial.GetMaterial();
 			}
 		}
 
@@ -166,9 +166,9 @@ namespace DatasmithRuntime
 		{
 			const TMap< FName, int32 >& TextureParams = GetMaterialParameters(Material).TextureParams;
 
-			for (int Index = 0; Index < MasterMaterialElement->GetPropertiesCount(); ++Index)
+			for (int Index = 0; Index < ReferenceMaterialElement->GetPropertiesCount(); ++Index)
 			{
-				const TSharedPtr< IDatasmithKeyValueProperty > Property = MasterMaterialElement->GetProperty(Index);
+				const TSharedPtr< IDatasmithKeyValueProperty > Property = ReferenceMaterialElement->GetProperty(Index);
 				const FName PropertyName(Property->GetName());
 
 				if (TextureParams.Contains(PropertyName))
@@ -185,28 +185,28 @@ namespace DatasmithRuntime
 		return MaterialRequirement;
 	}
 
-	bool LoadMasterMaterial(UMaterialInstanceDynamic* MaterialInstance, TSharedPtr<IDatasmithMasterMaterialElement>& MaterialElement)
+	bool LoadReferenceMaterial(UMaterialInstanceDynamic* MaterialInstance, TSharedPtr<IDatasmithMaterialInstanceElement>& MaterialElement)
 	{
-		TRACE_CPUPROFILER_EVENT_SCOPE(DatasmithRuntime::LoadMasterMaterial);
+		TRACE_CPUPROFILER_EVENT_SCOPE(DatasmithRuntime::LoadReferenceMaterial);
 
-		FDatasmithMasterMaterialManager& MaterialManager = FDatasmithMasterMaterialManager::Get();
+		FDatasmithReferenceMaterialManager& MaterialManager = FDatasmithReferenceMaterialManager::Get();
 		const FString Host = MaterialManager.GetHostFromString( MATERIAL_HOST );
-		TSharedPtr< FDatasmithMasterMaterialSelector > MaterialSelector = MaterialManager.GetSelector( MATERIAL_HOST );
+		TSharedPtr< FDatasmithReferenceMaterialSelector > MaterialSelector = MaterialManager.GetSelector( MATERIAL_HOST );
 
 		UMaterialInterface* ParentMaterial = nullptr;
 
 		{
-			if ( MaterialElement->GetMaterialType() == EDatasmithMasterMaterialType::Custom )
+			if ( MaterialElement->GetMaterialType() == EDatasmithReferenceMaterialType::Custom )
 			{
-				FDatasmithMasterMaterial CustomMasterMaterial;
+				FDatasmithReferenceMaterial CustomReferenceMaterial;
 
-				CustomMasterMaterial.FromSoftObjectPath( FSoftObjectPath( MaterialElement->GetCustomMaterialPathName() ) );
-				ParentMaterial = CustomMasterMaterial.GetMaterial();
+				CustomReferenceMaterial.FromSoftObjectPath( FSoftObjectPath( MaterialElement->GetCustomMaterialPathName() ) );
+				ParentMaterial = CustomReferenceMaterial.GetMaterial();
 			}
 			else if ( MaterialSelector.IsValid() )
 			{
-				const FDatasmithMasterMaterial& DatasmithMasterMaterial = MaterialSelector->GetMasterMaterial(MaterialElement);
-				ParentMaterial = DatasmithMasterMaterial.GetMaterial();
+				const FDatasmithReferenceMaterial& DatasmithReferenceMaterial = MaterialSelector->GetReferenceMaterial(MaterialElement);
+				ParentMaterial = DatasmithReferenceMaterial.GetMaterial();
 			}
 		}
 
diff --git a/Engine/Plugins/Experimental/Enterprise/DatasmithRuntime/Source/Private/MaterialImportUtils.h b/Engine/Plugins/Experimental/Enterprise/DatasmithRuntime/Source/Private/MaterialImportUtils.h
--- a/Engine/Plugins/Experimental/Enterprise/DatasmithRuntime/Source/Private/MaterialImportUtils.h
+++ b/Engine/Plugins/Experimental/Enterprise/DatasmithRuntime/Source/Private/MaterialImportUtils.h
@@ -6,7 +6,7 @@
 
 #include "SceneImporter.h"
 
-class IDatasmithMasterMaterialElement;
+class IDatasmithMaterialInstanceElement;
 class IDatasmithUEPbrMaterialElement;
 class UMaterialInstanceDynamic;
 
@@ -31,11 +31,11 @@ namespace DatasmithRuntime
 
 	TSharedPtr<IDatasmithUEPbrMaterialElement> ValidatePbrMaterial(TSharedPtr<IDatasmithUEPbrMaterialElement> PbrMaterialElement, FSceneImporter& SceneImporter);
 
-	int32 ProcessMaterialElement(TSharedPtr< IDatasmithMasterMaterialElement > BaseMaterialElement, FTextureCallback TextureCallback);
+	int32 ProcessMaterialElement(TSharedPtr< IDatasmithMaterialInstanceElement > BaseMaterialElement, FTextureCallback TextureCallback);
 
 	int32 ProcessMaterialElement(IDatasmithUEPbrMaterialElement* PbrMaterialElement, FTextureCallback TextureCallback);
 
-	bool LoadMasterMaterial(UMaterialInstanceDynamic* MaterialInstance, TSharedPtr<IDatasmithMasterMaterialElement>& MaterialElement);
+	bool LoadReferenceMaterial(UMaterialInstanceDynamic* MaterialInstance, TSharedPtr<IDatasmithMaterialInstanceElement>& MaterialElement);
 
 	bool LoadPbrMaterial(IDatasmithUEPbrMaterialElement& UEPbrMaterial, UMaterialInstanceDynamic* MaterialInstance);
 }
diff --git a/Engine/Plugins/Experimental/Enterprise/DatasmithRuntime/Source/Private/MaterialImporter.cpp b/Engine/Plugins/Experimental/Enterprise/DatasmithRuntime/Source/Private/MaterialImporter.cpp
--- a/Engine/Plugins/Experimental/Enterprise/DatasmithRuntime/Source/Private/MaterialImporter.cpp
+++ b/Engine/Plugins/Experimental/Enterprise/DatasmithRuntime/Source/Private/MaterialImporter.cpp
@@ -9,9 +9,9 @@
 #include "DatasmithMaterialElements.h"
 #include "DatasmithUtils.h"
 #include "IDatasmithSceneElements.h"
-#include "MasterMaterials/DatasmithMasterMaterial.h"
-#include "MasterMaterials/DatasmithMasterMaterialManager.h"
-#include "MasterMaterials/DatasmithMasterMaterialSelector.h"
+#include "ReferenceMaterials/DatasmithReferenceMaterial.h"
+#include "ReferenceMaterials/DatasmithReferenceMaterialManager.h"
+#include "ReferenceMaterials/DatasmithReferenceMaterialSelector.h"
 
 #include "Engine/Texture2D.h"
 #include "Engine/TextureLightProfile.h"
@@ -101,15 +101,15 @@ namespace DatasmithRuntime
 			}
 		};
 
-		const FString Host = FDatasmithMasterMaterialManager::Get().GetHostFromString( SceneElement->GetHost() );
+		const FString Host = FDatasmithReferenceMaterialManager::Get().GetHostFromString( SceneElement->GetHost() );
 
 		if( Element->IsA( EDatasmithElementType::UEPbrMaterial ) )
 		{
 			MaterialData.Requirements = ProcessMaterialElement(static_cast<IDatasmithUEPbrMaterialElement*>(Element.Get()), TextureCallback);
 		}
-		else if( Element->IsA( EDatasmithElementType::MasterMaterial ) )
+		else if( Element->IsA( EDatasmithElementType::MaterialInstance ) )
 		{
-			MaterialData.Requirements = ProcessMaterialElement(StaticCastSharedPtr<IDatasmithMasterMaterialElement>(Element), TextureCallback);
+			MaterialData.Requirements = ProcessMaterialElement(StaticCastSharedPtr<IDatasmithMaterialInstanceElement>(Element), TextureCallback);
 		}
 
 		MaterialData.SetState(EAssetState::Processed);
@@ -150,12 +150,12 @@ namespace DatasmithRuntime
 		{
 			// Not supported
 		}
-		else if ( Element->IsA( EDatasmithElementType::MasterMaterial ) )
+		else if ( Element->IsA( EDatasmithElementType::MaterialInstance ) )
 		{
 
-			TSharedPtr< IDatasmithMasterMaterialElement > MaterialElement = StaticCastSharedPtr< IDatasmithMasterMaterialElement >( Element );
+			TSharedPtr< IDatasmithMaterialInstanceElement > MaterialElement = StaticCastSharedPtr< IDatasmithMaterialInstanceElement >( Element );
 
-			bCreationSuccessful = LoadMasterMaterial(MaterialInstance, MaterialElement);
+			bCreationSuccessful = LoadReferenceMaterial(MaterialInstance, MaterialElement);
 
 			// Add tracking on material's properties
 			if (bCreationSuccessful)
@@ -221,9 +221,9 @@ namespace DatasmithRuntime
 			{
 				// Not supported
 			}
-			else if ( Element->IsA( EDatasmithElementType::MasterMaterial ) )
+			else if ( Element->IsA( EDatasmithElementType::MaterialInstance ) )
 			{
-				IDatasmithMasterMaterialElement* MaterialElement = static_cast< IDatasmithMasterMaterialElement* >( Element.Get() );
+				IDatasmithMaterialInstanceElement* MaterialElement = static_cast< IDatasmithMaterialInstanceElement* >( Element.Get() );
 
 				const int32 PropertyIndex = Referencer.Slot;
 				const TSharedPtr< IDatasmithKeyValueProperty > Property = MaterialElement->GetProperty(PropertyIndex);
diff --git a/Engine/Plugins/Experimental/Enterprise/DatasmithRuntime/Source/Private/MaterialSelectors/DatasmithRuntimeMaterialSelector.cpp b/Engine/Plugins/Experimental/Enterprise/DatasmithRuntime/Source/Private/MaterialSelectors/DatasmithRuntimeMaterialSelector.cpp
--- a/Engine/Plugins/Experimental/Enterprise/DatasmithRuntime/Source/Private/MaterialSelectors/DatasmithRuntimeMaterialSelector.cpp
+++ b/Engine/Plugins/Experimental/Enterprise/DatasmithRuntime/Source/Private/MaterialSelectors/DatasmithRuntimeMaterialSelector.cpp
@@ -21,9 +21,9 @@ bool FDatasmithRuntimeMaterialSelector::IsValid() const
 	return OpaqueMaterial.IsValid() && TransparentMaterial.IsValid() && CutoutMaterial.IsValid();
 }
 
-const FDatasmithMasterMaterial& FDatasmithRuntimeMaterialSelector::GetMasterMaterial( const TSharedPtr< IDatasmithMasterMaterialElement >& InDatasmithMaterial ) const
+const FDatasmithReferenceMaterial& FDatasmithRuntimeMaterialSelector::GetReferenceMaterial( const TSharedPtr< IDatasmithMaterialInstanceElement >& InDatasmithMaterial ) const
 {
-	TSharedPtr< IDatasmithMasterMaterialElement > MaterialElement = ConstCastSharedPtr< IDatasmithMasterMaterialElement >(InDatasmithMaterial);
+	TSharedPtr< IDatasmithMaterialInstanceElement > MaterialElement = ConstCastSharedPtr< IDatasmithMaterialInstanceElement >(InDatasmithMaterial);
 
 	TFunction<void(const TCHAR*, const TCHAR*)> ConvertProperty;
 	ConvertProperty = [this, &MaterialElement](const TCHAR* BoolPropertyName, const TCHAR* FloatPropertyName)
@@ -66,7 +66,7 @@ const FDatasmithMasterMaterial& FDatasmithRuntimeMaterialSelector::GetMasterMate
 		Roughness->SetValue(*NewValue);
 	}
 
-	// Convert static boolean parameters into float ones used in master material's graph
+	// Convert static boolean parameters into float ones used in reference material's graph
 	ConvertProperty(TEXT("RoughnessMapEnable"), TEXT("RoughnessMapFading"));
 	ConvertProperty(TEXT("IsMetal"), TEXT("Metallic"));
 	ConvertProperty(TEXT("TintEnabled"), TEXT("TintColorFading"));
@@ -74,11 +74,11 @@ const FDatasmithMasterMaterial& FDatasmithRuntimeMaterialSelector::GetMasterMate
 	ConvertProperty(TEXT("IsPbr"), TEXT("UseNormalMap"));
 
 	// Return proper material based on material's type
-	if (InDatasmithMaterial->GetMaterialType() == EDatasmithMasterMaterialType::Transparent)
+	if (InDatasmithMaterial->GetMaterialType() == EDatasmithReferenceMaterialType::Transparent)
 	{
 		return TransparentMaterial;
 	}
-	else if (InDatasmithMaterial->GetMaterialType() == EDatasmithMasterMaterialType::CutOut)
+	else if (InDatasmithMaterial->GetMaterialType() == EDatasmithReferenceMaterialType::CutOut)
 	{
 		return  CutoutMaterial;
 	}
@@ -86,7 +86,7 @@ const FDatasmithMasterMaterial& FDatasmithRuntimeMaterialSelector::GetMasterMate
 	return OpaqueMaterial;
 }
 
-void FDatasmithRuntimeMaterialSelector::FinalizeMaterialInstance(const TSharedPtr<IDatasmithMasterMaterialElement>& InDatasmithMaterial, UMaterialInstanceConstant * MaterialInstance) const
+void FDatasmithRuntimeMaterialSelector::FinalizeMaterialInstance(const TSharedPtr<IDatasmithMaterialInstanceElement>& InDatasmithMaterial, UMaterialInstanceConstant * MaterialInstance) const
 {
 	// Nothing to do there.
 }
diff --git a/Engine/Plugins/Experimental/Enterprise/DatasmithRuntime/Source/Private/MaterialSelectors/DatasmithRuntimeMaterialSelector.h b/Engine/Plugins/Experimental/Enterprise/DatasmithRuntime/Source/Private/MaterialSelectors/DatasmithRuntimeMaterialSelector.h
--- a/Engine/Plugins/Experimental/Enterprise/DatasmithRuntime/Source/Private/MaterialSelectors/DatasmithRuntimeMaterialSelector.h
+++ b/Engine/Plugins/Experimental/Enterprise/DatasmithRuntime/Source/Private/MaterialSelectors/DatasmithRuntimeMaterialSelector.h
@@ -2,25 +2,25 @@
 #pragma once
 
 #include "DatasmithDefinitions.h"
-#include "MasterMaterials/DatasmithMasterMaterial.h"
-#include "MasterMaterials/DatasmithMasterMaterialSelector.h"
+#include "ReferenceMaterials/DatasmithReferenceMaterial.h"
+#include "ReferenceMaterials/DatasmithReferenceMaterialSelector.h"
 
 #include "Templates/SharedPointer.h"
 
-class IDatasmithMasterMaterialElement;
+class IDatasmithMaterialInstanceElement;
 
-class FDatasmithRuntimeMaterialSelector : public FDatasmithMasterMaterialSelector
+class FDatasmithRuntimeMaterialSelector : public FDatasmithReferenceMaterialSelector
 {
 public:
 	FDatasmithRuntimeMaterialSelector();
 	virtual ~FDatasmithRuntimeMaterialSelector() = default;
 
 	virtual bool IsValid() const override;
-	virtual const FDatasmithMasterMaterial& GetMasterMaterial( const TSharedPtr< IDatasmithMasterMaterialElement >& /*InDatasmithMaterial*/ ) const override;
-	virtual void FinalizeMaterialInstance(const TSharedPtr< IDatasmithMasterMaterialElement >& InDatasmithMaterial, UMaterialInstanceConstant* MaterialInstance) const override;
+	virtual const FDatasmithReferenceMaterial& GetReferenceMaterial( const TSharedPtr< IDatasmithMaterialInstanceElement >& /*InDatasmithMaterial*/ ) const override;
+	virtual void FinalizeMaterialInstance(const TSharedPtr< IDatasmithMaterialInstanceElement >& InDatasmithMaterial, UMaterialInstanceConstant* MaterialInstance) const override;
 
 private:
-	FDatasmithMasterMaterial OpaqueMaterial;
-	FDatasmithMasterMaterial TransparentMaterial;
-	FDatasmithMasterMaterial CutoutMaterial;
+	FDatasmithReferenceMaterial OpaqueMaterial;
+	FDatasmithReferenceMaterial TransparentMaterial;
+	FDatasmithReferenceMaterial CutoutMaterial;
 };
diff --git a/Engine/Plugins/Experimental/Enterprise/DatasmithRuntime/Source/Private/SceneImporter.cpp b/Engine/Plugins/Experimental/Enterprise/DatasmithRuntime/Source/Private/SceneImporter.cpp
--- a/Engine/Plugins/Experimental/Enterprise/DatasmithRuntime/Source/Private/SceneImporter.cpp
+++ b/Engine/Plugins/Experimental/Enterprise/DatasmithRuntime/Source/Private/SceneImporter.cpp
@@ -1233,9 +1233,9 @@ namespace DatasmithRuntime
 			{
 				ProcessMaterialElement(static_cast<IDatasmithUEPbrMaterialElement*>(Element.Get()), TextureCallback);
 			}
-			else if( Element->IsA( EDatasmithElementType::MasterMaterial ) )
+			else if( Element->IsA( EDatasmithElementType::MaterialInstance ) )
 			{
-				ProcessMaterialElement(StaticCastSharedPtr<IDatasmithMasterMaterialElement>(Element), TextureCallback);
+				ProcessMaterialElement(StaticCastSharedPtr<IDatasmithMaterialInstanceElement>(Element), TextureCallback);
 			}
 
 			AssetPrefixedName = MaterialPrefix + Element->GetName();
diff --git a/Engine/Plugins/Experimental/Enterprise/DatasmithRuntime/Source/Public/DatasmithRuntime.h b/Engine/Plugins/Experimental/Enterprise/DatasmithRuntime/Source/Public/DatasmithRuntime.h
--- a/Engine/Plugins/Experimental/Enterprise/DatasmithRuntime/Source/Public/DatasmithRuntime.h
+++ b/Engine/Plugins/Experimental/Enterprise/DatasmithRuntime/Source/Public/DatasmithRuntime.h
@@ -19,7 +19,7 @@
 #include "DatasmithRuntime.generated.h"
 
 class ADatasmithRuntimeActor;
-class FDatasmithMasterMaterialSelector;
+class FDatasmithReferenceMaterialSelector;
 class FEvent;
 class IDatasmithScene;
 class IDatasmithTranslator;
