Updated Maxon libraries to support Apple new silicon

#rnx

diff --git a/Engine/Plugins/Enterprise/DatasmithC4DImporter/Source/DatasmithC4DTranslator/DatasmithC4DTranslator.Build.cs b/Engine/Plugins/Enterprise/DatasmithC4DImporter/Source/DatasmithC4DTranslator/DatasmithC4DTranslator.Build.cs
--- a/Engine/Plugins/Enterprise/DatasmithC4DImporter/Source/DatasmithC4DTranslator/DatasmithC4DTranslator.Build.cs
+++ b/Engine/Plugins/Enterprise/DatasmithC4DImporter/Source/DatasmithC4DTranslator/DatasmithC4DTranslator.Build.cs
@@ -9,6 +9,7 @@ namespace UnrealBuildTool.Rules
 		public DatasmithC4DTranslator(ReadOnlyTargetRules Target) : base(Target)
 		{
 			bLegalToDistributeObjectCode = true;
+			bUseUnity = false;
 
 			PrivateDependencyModuleNames.AddRange(
 				new string[]
diff --git a/Engine/Plugins/Enterprise/DatasmithC4DImporter/Source/DatasmithC4DTranslator/Private/DatasmithC4DExtraMelangeDefinitions.h b/Engine/Plugins/Enterprise/DatasmithC4DImporter/Source/DatasmithC4DTranslator/Private/DatasmithC4DExtraMelangeDefinitions.h
--- a/Engine/Plugins/Enterprise/DatasmithC4DImporter/Source/DatasmithC4DTranslator/Private/DatasmithC4DExtraMelangeDefinitions.h
+++ b/Engine/Plugins/Enterprise/DatasmithC4DImporter/Source/DatasmithC4DTranslator/Private/DatasmithC4DExtraMelangeDefinitions.h
@@ -2,9 +2,9 @@
 
 #pragma once
 
-// For some reason some symbols are not actually defined within melange, and are
+// For some reason some symbols are not actually defined within cineware, and are
 // passed around just as raw int32 ids.
-// Here we either define or inject them into the melange namespace like the others
+// Here we either define or inject them into the cineware namespace like the others
 
 #define Ocloner 1018544
 #define Ofracture 1018791
@@ -12,10 +12,10 @@
 #define REFLECTION_LAYER_MAIN_OPACITY REFLECTION_LAYER_TRANS_BRIGHTNESS
 #define REFLECTION_LAYER_NAME 0x0015
 
-namespace melange
+namespace cineware
 {
 	// We use C-style enums here instead of defines so that they are contained
-	// within the melange namespace and behave like the other enums
+	// within the cineware namespace and behave like the other enums
 
     enum // Tcrane attributes
     {
diff --git a/Engine/Plugins/Enterprise/DatasmithC4DImporter/Source/DatasmithC4DTranslator/Private/DatasmithC4DImporter.cpp b/Engine/Plugins/Enterprise/DatasmithC4DImporter/Source/DatasmithC4DTranslator/Private/DatasmithC4DImporter.cpp
--- a/Engine/Plugins/Enterprise/DatasmithC4DImporter/Source/DatasmithC4DTranslator/Private/DatasmithC4DImporter.cpp
+++ b/Engine/Plugins/Enterprise/DatasmithC4DImporter/Source/DatasmithC4DTranslator/Private/DatasmithC4DImporter.cpp
@@ -63,7 +63,7 @@ FDatasmithC4DImporter::FDatasmithC4DImporter(TSharedRef<IDatasmithScene>& OutSce
 
 FDatasmithC4DImporter::~FDatasmithC4DImporter()
 {
-	using namespace melange;
+	using namespace cineware;
 
 	if (C4dDocument != nullptr)
 	{
@@ -79,54 +79,54 @@ void FDatasmithC4DImporter::SetImportOptions(FDatasmithC4DImportOptions& InOptio
 
 namespace
 {
-	FMD5Hash ComputePolygonDataHash(melange::PolygonObject* PolyObject)
+	FMD5Hash ComputePolygonDataHash(cineware::PolygonObject* PolyObject)
 	{
-		melange::Int32 PointCount = PolyObject->GetPointCount();
-		melange::Int32 PolygonCount = PolyObject->GetPolygonCount();
-		const melange::Vector* Points = PolyObject->GetPointR();
-		const melange::CPolygon* Polygons = PolyObject->GetPolygonR();
-		melange::Vector32* Normals = PolyObject->CreatePhongNormals();
-		melange::GeData Data;
+		cineware::Int32 PointCount = PolyObject->GetPointCount();
+		cineware::Int32 PolygonCount = PolyObject->GetPolygonCount();
+		const cineware::Vector* Points = PolyObject->GetPointR();
+		const cineware::CPolygon* Polygons = PolyObject->GetPolygonR();
+		cineware::Vector32* Normals = PolyObject->CreatePhongNormals();
+		cineware::GeData Data;
 
 		FMD5 MD5;
-		MD5.Update(reinterpret_cast<const uint8*>(Points), sizeof(melange::Vector)*PointCount);
-		MD5.Update(reinterpret_cast<const uint8*>(Polygons), sizeof(melange::CPolygon)*PolygonCount);
+		MD5.Update(reinterpret_cast<const uint8*>(Points), sizeof(cineware::Vector)*PointCount);
+		MD5.Update(reinterpret_cast<const uint8*>(Polygons), sizeof(cineware::CPolygon)*PolygonCount);
 		if (Normals)
 		{
-			MD5.Update(reinterpret_cast<const uint8*>(Normals), sizeof(melange::Vector32)*PointCount);
-			melange::DeleteMem(Normals);
+			MD5.Update(reinterpret_cast<const uint8*>(Normals), sizeof(cineware::Vector32)*PointCount);
+			cineware::DeleteMem(Normals);
 		}
 
 		//Tags
-		for (melange::BaseTag* Tag = PolyObject->GetFirstTag(); Tag; Tag = Tag->GetNext())
+		for (cineware::BaseTag* Tag = PolyObject->GetFirstTag(); Tag; Tag = Tag->GetNext())
 		{
-			melange::Int32 TagType = Tag->GetType();
+			cineware::Int32 TagType = Tag->GetType();
 			if (TagType == Tuvw)
 			{
-				melange::ConstUVWHandle UVWHandle = static_cast<melange::UVWTag*>(Tag)->GetDataAddressR();
+				cineware::ConstUVWHandle UVWHandle = static_cast<cineware::UVWTag*>(Tag)->GetDataAddressR();
 				for (int32 PolygonIndex = 0; PolygonIndex < PolygonCount; ++PolygonIndex)
 				{
-					melange::UVWStruct UVWStruct;
-					melange::UVWTag::Get(UVWHandle, PolygonIndex, UVWStruct);
-					MD5.Update(reinterpret_cast<const uint8*>(&UVWStruct), sizeof(melange::UVWStruct));
+					cineware::UVWStruct UVWStruct;
+					cineware::UVWTag::Get(UVWHandle, PolygonIndex, UVWStruct);
+					MD5.Update(reinterpret_cast<const uint8*>(&UVWStruct), sizeof(cineware::UVWStruct));
 				}
 			}
 			else if (TagType == Tpolygonselection)
 			{
-				melange::SelectionTag* SelectionTag = static_cast<melange::SelectionTag*>(Tag);
-				melange::BaseSelect* BaseSelect = SelectionTag->GetBaseSelect();
+				cineware::SelectionTag* SelectionTag = static_cast<cineware::SelectionTag*>(Tag);
+				cineware::BaseSelect* BaseSelect = SelectionTag->GetBaseSelect();
 
-				FString SelectionName = MelangeGetString(SelectionTag, melange::POLYGONSELECTIONTAG_NAME);
+				FString SelectionName = MelangeGetString(SelectionTag, cineware::POLYGONSELECTIONTAG_NAME);
 				uint32 NameHash = GetTypeHash(SelectionName);
 				MD5.Update(reinterpret_cast<const uint8*>(&NameHash), sizeof(NameHash));
 
-				TArray<melange::Int32> PolygonSelections;
+				TArray<cineware::Int32> PolygonSelections;
 				PolygonSelections.Reserve(BaseSelect->GetCount());
 
-				melange::Int32 Segment = 0;
-				melange::Int32 RangeStart = 0;
-				melange::Int32 RangeEnd = 0;
-				melange::Int32 Selection = 0;
+				cineware::Int32 Segment = 0;
+				cineware::Int32 RangeStart = 0;
+				cineware::Int32 RangeEnd = 0;
+				cineware::Int32 Selection = 0;
 				while (BaseSelect->GetRange(Segment++, &RangeStart, &RangeEnd))
 				{
 					for (Selection = RangeStart; Selection <= RangeEnd; ++Selection)
@@ -134,7 +134,7 @@ namespace
 						PolygonSelections.Add(Selection);
 					}
 				}
-				MD5.Update(reinterpret_cast<const uint8*>(PolygonSelections.GetData()), PolygonSelections.Num() * sizeof(melange::Int32));
+				MD5.Update(reinterpret_cast<const uint8*>(PolygonSelections.GetData()), PolygonSelections.Num() * sizeof(cineware::Int32));
 			}
 		}
 
@@ -169,40 +169,40 @@ struct FCraneCameraAttributes
 	{
 		switch (AttributeID)
 		{
-		case melange::CRANECAMERA_BASE_HEIGHT:
+		case cineware::CRANECAMERA_BASE_HEIGHT:
 			BaseHeight = (float)AttributeValue;
 			break;
-		case melange::CRANECAMERA_BASE_HEADING:
+		case cineware::CRANECAMERA_BASE_HEADING:
 			BaseHeading = (float)FMath::RadiansToDegrees(AttributeValue);
 			break;
-		case melange::CRANECAMERA_ARM_LENGTH:
+		case cineware::CRANECAMERA_ARM_LENGTH:
 			ArmLength = (float)AttributeValue;
 			break;
-		case melange::CRANECAMERA_ARM_PITCH:
+		case cineware::CRANECAMERA_ARM_PITCH:
 			ArmPitch = (float)FMath::RadiansToDegrees(AttributeValue);
 			break;
-		case melange::CRANECAMERA_HEAD_HEIGHT:
+		case cineware::CRANECAMERA_HEAD_HEIGHT:
 			HeadHeight = (float)AttributeValue;
 			break;
-		case melange::CRANECAMERA_HEAD_HEADING:
+		case cineware::CRANECAMERA_HEAD_HEADING:
 			HeadHeading = (float)FMath::RadiansToDegrees(AttributeValue);
 			break;
-		case melange::CRANECAMERA_HEAD_WIDTH:
+		case cineware::CRANECAMERA_HEAD_WIDTH:
 			HeadWidth = (float)AttributeValue;
 			break;
-		case melange::CRANECAMERA_CAM_PITCH:
+		case cineware::CRANECAMERA_CAM_PITCH:
 			CamPitch = (float)FMath::RadiansToDegrees(AttributeValue);
 			break;
-		case melange::CRANECAMERA_CAM_BANKING:
+		case cineware::CRANECAMERA_CAM_BANKING:
 			CamBanking = (float)FMath::RadiansToDegrees(AttributeValue);
 			break;
-		case melange::CRANECAMERA_CAM_OFFSET:
+		case cineware::CRANECAMERA_CAM_OFFSET:
 			CamOffset = (float)AttributeValue;
 			break;
-		case melange::CRANECAMERA_COMPENSATE_PITCH:
+		case cineware::CRANECAMERA_COMPENSATE_PITCH:
 			bCompensatePitch = static_cast<bool>(AttributeValue);
 			break;
-		case melange::CRANECAMERA_COMPENSATE_HEADING:
+		case cineware::CRANECAMERA_COMPENSATE_HEADING:
 			bCompensateHeading = static_cast<bool>(AttributeValue);
 			break;
 		default:
@@ -212,58 +212,58 @@ struct FCraneCameraAttributes
 };
 
 // Extracts all of the relevant parameters from a Tcrane tag and packs them in a FCraneCameraAttributes
-TSharedRef<FCraneCameraAttributes> ExtractCraneCameraAttributes(melange::BaseTag* CraneTag)
+TSharedRef<FCraneCameraAttributes> ExtractCraneCameraAttributes(cineware::BaseTag* CraneTag)
 {
 	TSharedRef<FCraneCameraAttributes> Result = MakeShared<FCraneCameraAttributes>();
 
-	melange::GeData Data;
-	if (CraneTag->GetParameter(melange::CRANECAMERA_BASE_HEIGHT, Data))
+	cineware::GeData Data;
+	if (CraneTag->GetParameter(cineware::CRANECAMERA_BASE_HEIGHT, Data))
 	{
-		Result->SetAttributeByID(melange::CRANECAMERA_BASE_HEIGHT, Data.GetFloat());
+		Result->SetAttributeByID(cineware::CRANECAMERA_BASE_HEIGHT, Data.GetFloat());
 	}
-	if (CraneTag->GetParameter(melange::CRANECAMERA_BASE_HEADING, Data))
+	if (CraneTag->GetParameter(cineware::CRANECAMERA_BASE_HEADING, Data))
 	{
-		Result->SetAttributeByID(melange::CRANECAMERA_BASE_HEADING, Data.GetFloat());
+		Result->SetAttributeByID(cineware::CRANECAMERA_BASE_HEADING, Data.GetFloat());
 	}
-	if (CraneTag->GetParameter(melange::CRANECAMERA_ARM_LENGTH, Data))
+	if (CraneTag->GetParameter(cineware::CRANECAMERA_ARM_LENGTH, Data))
 	{
-		Result->SetAttributeByID(melange::CRANECAMERA_ARM_LENGTH, Data.GetFloat());
+		Result->SetAttributeByID(cineware::CRANECAMERA_ARM_LENGTH, Data.GetFloat());
 	}
-	if (CraneTag->GetParameter(melange::CRANECAMERA_ARM_PITCH, Data))
+	if (CraneTag->GetParameter(cineware::CRANECAMERA_ARM_PITCH, Data))
 	{
-		Result->SetAttributeByID(melange::CRANECAMERA_ARM_PITCH, Data.GetFloat());
+		Result->SetAttributeByID(cineware::CRANECAMERA_ARM_PITCH, Data.GetFloat());
 	}
-	if (CraneTag->GetParameter(melange::CRANECAMERA_HEAD_HEIGHT, Data))
+	if (CraneTag->GetParameter(cineware::CRANECAMERA_HEAD_HEIGHT, Data))
 	{
-		Result->SetAttributeByID(melange::CRANECAMERA_HEAD_HEIGHT, Data.GetFloat());
+		Result->SetAttributeByID(cineware::CRANECAMERA_HEAD_HEIGHT, Data.GetFloat());
 	}
-	if (CraneTag->GetParameter(melange::CRANECAMERA_HEAD_HEADING, Data))
+	if (CraneTag->GetParameter(cineware::CRANECAMERA_HEAD_HEADING, Data))
 	{
-		Result->SetAttributeByID(melange::CRANECAMERA_HEAD_HEADING, Data.GetFloat());
+		Result->SetAttributeByID(cineware::CRANECAMERA_HEAD_HEADING, Data.GetFloat());
 	}
-	if (CraneTag->GetParameter(melange::CRANECAMERA_HEAD_WIDTH, Data))
+	if (CraneTag->GetParameter(cineware::CRANECAMERA_HEAD_WIDTH, Data))
 	{
-		Result->SetAttributeByID(melange::CRANECAMERA_HEAD_WIDTH, Data.GetFloat());
+		Result->SetAttributeByID(cineware::CRANECAMERA_HEAD_WIDTH, Data.GetFloat());
 	}
-	if (CraneTag->GetParameter(melange::CRANECAMERA_CAM_PITCH, Data))
+	if (CraneTag->GetParameter(cineware::CRANECAMERA_CAM_PITCH, Data))
 	{
-		Result->SetAttributeByID(melange::CRANECAMERA_CAM_PITCH, Data.GetFloat());
+		Result->SetAttributeByID(cineware::CRANECAMERA_CAM_PITCH, Data.GetFloat());
 	}
-	if (CraneTag->GetParameter(melange::CRANECAMERA_CAM_BANKING, Data))
+	if (CraneTag->GetParameter(cineware::CRANECAMERA_CAM_BANKING, Data))
 	{
-		Result->SetAttributeByID(melange::CRANECAMERA_CAM_BANKING, Data.GetFloat());
+		Result->SetAttributeByID(cineware::CRANECAMERA_CAM_BANKING, Data.GetFloat());
 	}
-	if (CraneTag->GetParameter(melange::CRANECAMERA_CAM_OFFSET, Data))
+	if (CraneTag->GetParameter(cineware::CRANECAMERA_CAM_OFFSET, Data))
 	{
-		Result->SetAttributeByID(melange::CRANECAMERA_CAM_OFFSET, Data.GetFloat());
+		Result->SetAttributeByID(cineware::CRANECAMERA_CAM_OFFSET, Data.GetFloat());
 	}
-	if (CraneTag->GetParameter(melange::CRANECAMERA_COMPENSATE_PITCH, Data))
+	if (CraneTag->GetParameter(cineware::CRANECAMERA_COMPENSATE_PITCH, Data))
 	{
-		Result->SetAttributeByID(melange::CRANECAMERA_COMPENSATE_PITCH, Data.GetInt32());
+		Result->SetAttributeByID(cineware::CRANECAMERA_COMPENSATE_PITCH, Data.GetInt32());
 	}
-	if (CraneTag->GetParameter(melange::CRANECAMERA_COMPENSATE_HEADING, Data))
+	if (CraneTag->GetParameter(cineware::CRANECAMERA_COMPENSATE_HEADING, Data))
 	{
-		Result->SetAttributeByID(melange::CRANECAMERA_COMPENSATE_HEADING, Data.GetInt32());
+		Result->SetAttributeByID(cineware::CRANECAMERA_COMPENSATE_HEADING, Data.GetInt32());
 	}
 	return Result;
 }
@@ -318,7 +318,7 @@ FTransform CalculateCraneCameraTransform(const FCraneCameraAttributes& Params)
 	FVector TranslationUE = FinalTransUE.GetTranslation();
 	FVector3f EulerUE = (FVector3f)FinalTransUE.GetRotation().Euler();
 
-	// Convert FinalTransUnrealEditor into the melange coordinate system, so that this can be treated
+	// Convert FinalTransUnrealEditor into the cineware coordinate system, so that this can be treated
 	// like the other types of animations in ImportAnimations.
 	// More specifically, convert them so that that ConvertDirectionLeftHandedYup and
 	// the conversion for Ocamera rotations gets them back into UnrealEditor's coordinate system
@@ -327,11 +327,11 @@ FTransform CalculateCraneCameraTransform(const FCraneCameraAttributes& Params)
 					  FVector(TranslationUE.X, TranslationUE.Z, -TranslationUE.Y));
 }
 
-void FDatasmithC4DImporter::ImportSpline(melange::SplineObject* SplineActor)
+void FDatasmithC4DImporter::ImportSpline(cineware::SplineObject* SplineActor)
 {
 	// ActorObject has fewer keys, but uses bezier control points
 	// Cache has more keys generated by subdivision, should be parsed with linear interpolation
-	melange::SplineObject* SplineCache = static_cast<melange::SplineObject*>(GetBestMelangeCache(SplineActor));
+	cineware::SplineObject* SplineCache = static_cast<cineware::SplineObject*>(GetBestMelangeCache(SplineActor));
 
 	if (SplineActor && SplineCache)
 	{
@@ -356,12 +356,12 @@ void FDatasmithC4DImporter::ImportSpline(melange::SplineObject* SplineActor)
 			++PercentageDenominator;
 		}
 
-		melange::Matrix Trans = SplineCache->GetMg();
+		cineware::Matrix Trans = SplineCache->GetMg();
 
-		const melange::Vector* Points = SplineCache->GetPointR();
+		const cineware::Vector* Points = SplineCache->GetPointR();
 		for (int32 PointIndex = 0; PointIndex < NumPoints; ++PointIndex)
 		{
-			const melange::Vector& Point = Trans * Points[PointIndex];
+			const cineware::Vector& Point = Trans * Points[PointIndex];
 			float Percent = PointIndex / PercentageDenominator;
 			XCurve.AddKey(Percent, (float)Point.x);
 			YCurve.AddKey(Percent, (float)Point.y);
@@ -370,7 +370,7 @@ void FDatasmithC4DImporter::ImportSpline(melange::SplineObject* SplineActor)
 
 		if (SplineActor->GetIsClosed())
 		{
-			const melange::Vector& FirstPoint = Trans * Points[0];
+			const cineware::Vector& FirstPoint = Trans * Points[0];
 			XCurve.AddKey(1.0f, (float)FirstPoint.x);
 			YCurve.AddKey(1.0f, (float)FirstPoint.y);
 			ZCurve.AddKey(1.0f, (float)FirstPoint.z);
@@ -378,7 +378,7 @@ void FDatasmithC4DImporter::ImportSpline(melange::SplineObject* SplineActor)
 	}
 }
 
-melange::BaseObject* FDatasmithC4DImporter::GetBestMelangeCache(melange::BaseObject* Object)
+cineware::BaseObject* FDatasmithC4DImporter::GetBestMelangeCache(cineware::BaseObject* Object)
 {
 	if (Object == nullptr)
 	{
@@ -387,7 +387,7 @@ melange::BaseObject* FDatasmithC4DImporter::GetBestMelangeCache(melange::BaseObj
 
 	//When primitives types (cube, cone, cylinder...) are exported with the "Save Project for Melange" option,
 	//they will have a cache that represent their PolygonObject equivalent.
-	melange::BaseObject* ObjectCache = Object->GetCache();
+	cineware::BaseObject* ObjectCache = Object->GetCache();
 
 	//When the primitive has a deformer, the resulting PolygonObject will be in a sub-cache
 	if (ObjectCache)
@@ -410,16 +410,16 @@ melange::BaseObject* FDatasmithC4DImporter::GetBestMelangeCache(melange::BaseObj
 	return ObjectCache;
 }
 
-TOptional<FString> FDatasmithC4DImporter::MelangeObjectID(melange::BaseObject* Object)
+TOptional<FString> FDatasmithC4DImporter::MelangeObjectID(cineware::BaseObject* Object)
 {
 	//Make sure that Object is not in a cache
 	FString HierarchyPosition;
 	bool InCache = false;
-	melange::BaseObject* ParentObject = Object;
+	cineware::BaseObject* ParentObject = Object;
 	while (ParentObject)
 	{
 		int ObjectHierarchyIndex = 0;
-		melange::BaseObject* PrevObject = ParentObject->GetPred();
+		cineware::BaseObject* PrevObject = ParentObject->GetPred();
 		while (PrevObject)
 		{
 			ObjectHierarchyIndex++;
@@ -427,7 +427,7 @@ TOptional<FString> FDatasmithC4DImporter::MelangeObjectID(melange::BaseObject* O
 		}
 		HierarchyPosition = "_" + FString::FromInt(ObjectHierarchyIndex) + HierarchyPosition;
 
-		melange::BaseObject** OriginalObject = CachesOriginalObject.Find(ParentObject);
+		cineware::BaseObject** OriginalObject = CachesOriginalObject.Find(ParentObject);
 		if (OriginalObject)
 		{
 			InCache = true;
@@ -608,9 +608,9 @@ namespace FC4DImporterImpl
 		Metadata->AddProperty(MetadataPropertyPtr);
 	}
 
-	void ImportActorMetadata(melange::BaseObject* Object, const TSharedPtr<IDatasmithActorElement>& Actor, TSharedRef<IDatasmithScene>& DatasmithScene)
+	void ImportActorMetadata(cineware::BaseObject* Object, const TSharedPtr<IDatasmithActorElement>& Actor, TSharedRef<IDatasmithScene>& DatasmithScene)
 	{
-		melange::DynamicDescription* DynamicDescription = Object->GetDynamicDescription();
+		cineware::DynamicDescription* DynamicDescription = Object->GetDynamicDescription();
 		if (!DynamicDescription)
 		{
 			return;
@@ -619,11 +619,11 @@ namespace FC4DImporterImpl
 		TSharedPtr<IDatasmithMetaDataElement> Metadata = nullptr;
 
 		void* BrowserHandle = DynamicDescription->BrowseInit();
-		melange::DescID DescId;
-		const melange::BaseContainer* DescContainer;
+		cineware::DescID DescId;
+		const cineware::BaseContainer* DescContainer;
 		while (DynamicDescription->BrowseGetNext(BrowserHandle, &DescId, &DescContainer))
 		{
-			if (DescId[0].id != melange::ID_USERDATA)
+			if (DescId[0].id != cineware::ID_USERDATA)
 			{
 				continue;;
 			}
@@ -637,21 +637,21 @@ namespace FC4DImporterImpl
 				}
 			}
 
-			melange::GeData Data;
+			cineware::GeData Data;
 			if (!Object->GetParameter(DescId, Data))
 			{
 				continue;
 			}
 
-			FString DataName = MelangeStringToFString(DescContainer->GetString(melange::DESC_NAME));
+			FString DataName = MelangeStringToFString(DescContainer->GetString(cineware::DESC_NAME));
 
-			melange::Int32 UserDataType = DescContainer->GetInt32(melange::DESC_CUSTOMGUI);
-			if (UserDataType == melange::DA_VECTOR)
+			cineware::Int32 UserDataType = DescContainer->GetInt32(cineware::DESC_CUSTOMGUI);
+			if (UserDataType == cineware::DA_VECTOR)
 			{
 				FVector ConvertedVector = ConvertMelangePosition(Data.GetVector());
 				AddMetadataVector(Metadata.Get(), *DataName, ConvertedVector);
 			}
-			else if (UserDataType == melange::DA_REAL)
+			else if (UserDataType == cineware::DA_REAL)
 			{
 				AddMetadataFloat(Metadata.Get(), *DataName, static_cast<float>(Data.GetFloat()));
 			}
@@ -680,7 +680,7 @@ namespace FC4DImporterImpl
 	}
 }
 
-bool FDatasmithC4DImporter::AddChildActor(melange::BaseObject* Object, TSharedPtr<IDatasmithActorElement> ParentActor, melange::Matrix WorldTransformMatrix, const TSharedPtr<IDatasmithActorElement>& Actor)
+bool FDatasmithC4DImporter::AddChildActor(cineware::BaseObject* Object, TSharedPtr<IDatasmithActorElement> ParentActor, cineware::Matrix WorldTransformMatrix, const TSharedPtr<IDatasmithActorElement>& Actor)
 {
 	FC4DImporterImpl::ImportActorMetadata(Object, Actor, DatasmithScene);
 
@@ -697,11 +697,11 @@ bool FDatasmithC4DImporter::AddChildActor(melange::BaseObject* Object, TSharedPt
 	{
 		// Compensates the fact that in C4D cameras/lights shoot out towards +Z, while in
 		// UnrealEditor they shoot towards +X
-		melange::Matrix CameraRotation(
-			melange::Vector(0.0, 0.0, 0.0),
-			melange::Vector(0.0, 0.0, 1.0),
-			melange::Vector(0.0, 1.0, 0.0),
-			melange::Vector(-1.0, 0.0, 0.0));
+		cineware::Matrix CameraRotation(
+			cineware::Vector(0.0, 0.0, 0.0),
+			cineware::Vector(0.0, 0.0, 1.0),
+			cineware::Vector(0.0, 1.0, 0.0),
+			cineware::Vector(-1.0, 0.0, 0.0));
 		WorldTransformMatrix = WorldTransformMatrix * CameraRotation;
 	}
 
@@ -794,7 +794,7 @@ bool FDatasmithC4DImporter::AddChildActor(melange::BaseObject* Object, TSharedPt
 	return true;
 }
 
-TSharedPtr<IDatasmithActorElement> FDatasmithC4DImporter::ImportNullActor(melange::BaseObject* Object, const FString& DatasmithName, const FString& DatasmithLabel)
+TSharedPtr<IDatasmithActorElement> FDatasmithC4DImporter::ImportNullActor(cineware::BaseObject* Object, const FString& DatasmithName, const FString& DatasmithLabel)
 {
 	TSharedPtr<IDatasmithActorElement> ActorElement = FDatasmithSceneFactory::CreateActor(*DatasmithName);
 	ActorElement->SetLabel(*DatasmithLabel);
@@ -808,31 +808,31 @@ namespace
 		TSharedPtr<IDatasmithLightActorElement> Result = nullptr;
 		switch (MelangeLightTypeId)
 		{
-		case melange::LIGHT_TYPE_OMNI:
+		case cineware::LIGHT_TYPE_OMNI:
 			Result = FDatasmithSceneFactory::CreatePointLight(*Name);
 			break;
-		case melange::LIGHT_TYPE_SPOT:
+		case cineware::LIGHT_TYPE_SPOT:
 			Result = FDatasmithSceneFactory::CreateSpotLight(*Name);
 			break;
-		case melange::LIGHT_TYPE_SPOTRECT:
+		case cineware::LIGHT_TYPE_SPOTRECT:
 			Result = FDatasmithSceneFactory::CreateSpotLight(*Name);
 			break;
-		case melange::LIGHT_TYPE_DISTANT:
+		case cineware::LIGHT_TYPE_DISTANT:
 			Result = FDatasmithSceneFactory::CreateDirectionalLight(*Name);
 			break;
-		case melange::LIGHT_TYPE_PARALLEL:
+		case cineware::LIGHT_TYPE_PARALLEL:
 			Result = FDatasmithSceneFactory::CreateSpotLight(*Name);
 			break;
-		case melange::LIGHT_TYPE_PARSPOTRECT:
+		case cineware::LIGHT_TYPE_PARSPOTRECT:
 			Result = FDatasmithSceneFactory::CreateSpotLight(*Name);
 			break;
-		case melange::LIGHT_TYPE_TUBE:
+		case cineware::LIGHT_TYPE_TUBE:
 			Result = FDatasmithSceneFactory::CreateSpotLight(*Name);
 			break;
-		case melange::LIGHT_TYPE_AREA:
+		case cineware::LIGHT_TYPE_AREA:
 			Result = FDatasmithSceneFactory::CreateAreaLight(*Name);
 			break;
-		case melange::LIGHT_TYPE_PHOTOMETRIC:
+		case cineware::LIGHT_TYPE_PHOTOMETRIC:
 			Result = FDatasmithSceneFactory::CreatePointLight(*Name);
 			break;
 		default:
@@ -851,10 +851,10 @@ namespace
 	{
 		switch (MelangeLightUnitId)
 		{
-		case melange::LIGHT_PHOTOMETRIC_UNIT_LUMEN:
+		case cineware::LIGHT_PHOTOMETRIC_UNIT_LUMEN:
 			return EDatasmithLightUnits::Lumens;
 			break;
-		case melange::LIGHT_PHOTOMETRIC_UNIT_CANDELA:
+		case cineware::LIGHT_PHOTOMETRIC_UNIT_CANDELA:
 			return EDatasmithLightUnits::Candelas;
 			break;
 		default:
@@ -869,31 +869,31 @@ namespace
 	{
 		switch (AreaLightC4DId)
 		{
-		case melange::LIGHT_AREADETAILS_SHAPE_DISC:
+		case cineware::LIGHT_AREADETAILS_SHAPE_DISC:
 			return EDatasmithLightShape::Disc;
 			break;
-		case melange::LIGHT_AREADETAILS_SHAPE_RECTANGLE:
+		case cineware::LIGHT_AREADETAILS_SHAPE_RECTANGLE:
 			return EDatasmithLightShape::Rectangle;
 			break;
-		case melange::LIGHT_AREADETAILS_SHAPE_SPHERE:
+		case cineware::LIGHT_AREADETAILS_SHAPE_SPHERE:
 			return EDatasmithLightShape::Sphere;
 			break;
-		case melange::LIGHT_AREADETAILS_SHAPE_CYLINDER:
+		case cineware::LIGHT_AREADETAILS_SHAPE_CYLINDER:
 			return EDatasmithLightShape::Cylinder;
 			break;
-		case melange::LIGHT_AREADETAILS_SHAPE_CUBE:
+		case cineware::LIGHT_AREADETAILS_SHAPE_CUBE:
 			return EDatasmithLightShape::Rectangle;
 			break;
-		case melange::LIGHT_AREADETAILS_SHAPE_HEMISPHERE:
+		case cineware::LIGHT_AREADETAILS_SHAPE_HEMISPHERE:
 			return EDatasmithLightShape::Sphere;
 			break;
-		case melange::LIGHT_AREADETAILS_SHAPE_OBJECT:
+		case cineware::LIGHT_AREADETAILS_SHAPE_OBJECT:
 			return EDatasmithLightShape::None;
 			break;
-		case melange::LIGHT_AREADETAILS_SHAPE_LINE:
+		case cineware::LIGHT_AREADETAILS_SHAPE_LINE:
 			return EDatasmithLightShape::Cylinder;
 			break;
-		case melange::LIGHT_AREADETAILS_SHAPE_PCYLINDER:
+		case cineware::LIGHT_AREADETAILS_SHAPE_PCYLINDER:
 			return EDatasmithLightShape::Cylinder;
 			break;
 		default:
@@ -906,12 +906,12 @@ namespace
 
 namespace
 {
-	melange::Int32 MelangeColorProfile = melange::DOCUMENT_COLORPROFILE_SRGB;
+	cineware::Int32 MelangeColorProfile = cineware::DOCUMENT_COLORPROFILE_SRGB;
 
 	FVector ToLinearColor(const FVector& Color)
 	{
 		// Document is already linear, nothing to do
-		if (MelangeColorProfile == melange::DOCUMENT_COLORPROFILE_LINEAR)
+		if (MelangeColorProfile == cineware::DOCUMENT_COLORPROFILE_LINEAR)
 		{
 			return Color;
 		}
@@ -922,7 +922,7 @@ namespace
 	}
 
 	// Gets a color weighted by its brightness
-	FVector MelangeGetLayerColor(melange::BaseList2D* MelangeObject, melange::Int32 ColorAttributeID, melange::Int32 BrightnessAttributeID)
+	FVector MelangeGetLayerColor(cineware::BaseList2D* MelangeObject, cineware::Int32 ColorAttributeID, cineware::Int32 BrightnessAttributeID)
 	{
 		FVector Result;
 		if (MelangeObject)
@@ -935,7 +935,7 @@ namespace
 	}
 
 	// In here instead of utils because it depends on the document color profile
-	FVector MelangeGetColor(melange::BaseList2D* MelangeObject, melange::Int32 MelangeDescId)
+	FVector MelangeGetColor(cineware::BaseList2D* MelangeObject, cineware::Int32 MelangeDescId)
 	{
 		FVector Result;
 		if (MelangeObject)
@@ -987,12 +987,12 @@ namespace
 	}
 }
 
-TSharedPtr<IDatasmithLightActorElement> FDatasmithC4DImporter::ImportLight(melange::BaseObject* InC4DLightPtr, const FString& DatasmithName, const FString& DatasmithLabel)
+TSharedPtr<IDatasmithLightActorElement> FDatasmithC4DImporter::ImportLight(cineware::BaseObject* InC4DLightPtr, const FString& DatasmithName, const FString& DatasmithLabel)
 {
-	melange::GeData C4DData;
+	cineware::GeData C4DData;
 
 	// Actor type
-	int32 LightTypeId = MelangeGetInt32(InC4DLightPtr, melange::LIGHT_TYPE);
+	int32 LightTypeId = MelangeGetInt32(InC4DLightPtr, cineware::LIGHT_TYPE);
 	TSharedPtr<IDatasmithLightActorElement> LightActor = CreateDatasmithLightActorElement(LightTypeId, DatasmithName, DatasmithLabel);
 	if (!LightActor.IsValid())
 	{
@@ -1001,11 +1001,11 @@ TSharedPtr<IDatasmithLightActorElement> FDatasmithC4DImporter::ImportLight(melan
 	}
 
 	// Color
-	FLinearColor Color = FLinearColor(MelangeGetColor(InC4DLightPtr, melange::LIGHT_COLOR));
+	FLinearColor Color = FLinearColor(MelangeGetColor(InC4DLightPtr, cineware::LIGHT_COLOR));
 
 	// Temperature
-	bool bUseTemperature = MelangeGetBool(InC4DLightPtr, melange::LIGHT_TEMPERATURE);
-	double Temperature = MelangeGetDouble(InC4DLightPtr, melange::LIGHT_TEMPERATURE_MAIN);
+	bool bUseTemperature = MelangeGetBool(InC4DLightPtr, cineware::LIGHT_TEMPERATURE);
+	double Temperature = MelangeGetDouble(InC4DLightPtr, cineware::LIGHT_TEMPERATURE_MAIN);
 	if (Temperature == 0)
 	{
 		Temperature = 6500.0;
@@ -1014,15 +1014,15 @@ TSharedPtr<IDatasmithLightActorElement> FDatasmithC4DImporter::ImportLight(melan
 	// Intensity and units
 	double Intensity = 1.0;
 	EDatasmithLightUnits Units = EDatasmithLightUnits::Unitless;
-	if (MelangeGetBool(InC4DLightPtr, melange::LIGHT_PHOTOMETRIC_UNITS))
+	if (MelangeGetBool(InC4DLightPtr, cineware::LIGHT_PHOTOMETRIC_UNITS))
 	{
-		Units = GetDatasmithLightIntensityUnits(MelangeGetInt32(InC4DLightPtr, melange::LIGHT_PHOTOMETRIC_UNIT));
+		Units = GetDatasmithLightIntensityUnits(MelangeGetInt32(InC4DLightPtr, cineware::LIGHT_PHOTOMETRIC_UNIT));
 
-		Intensity = MelangeGetDouble(InC4DLightPtr, melange::LIGHT_PHOTOMETRIC_INTENSITY); // Cd/lm value in 'Photometric' tab
+		Intensity = MelangeGetDouble(InC4DLightPtr, cineware::LIGHT_PHOTOMETRIC_INTENSITY); // Cd/lm value in 'Photometric' tab
 	}
 
 	// Brightness
-	Intensity *= MelangeGetDouble(InC4DLightPtr, melange::LIGHT_BRIGHTNESS); // percentage value on 'General' tab, usually = 1.0
+	Intensity *= MelangeGetDouble(InC4DLightPtr, cineware::LIGHT_BRIGHTNESS); // percentage value on 'General' tab, usually = 1.0
 	if (Units == EDatasmithLightUnits::Unitless)
 	{
 		if (LightActor->IsA(EDatasmithElementType::PointLight))
@@ -1040,10 +1040,10 @@ TSharedPtr<IDatasmithLightActorElement> FDatasmithC4DImporter::ImportLight(melan
 	// Apparently non-IES lights can have this checked while the checkbox is in a "disabled state", so we must also check the light type
 	FString IESPath;
 	TOptional<double> IESBrightnessScale;
-	bool bUseIES = LightTypeId == melange::LIGHT_TYPE_PHOTOMETRIC && MelangeGetBool(InC4DLightPtr, melange::LIGHT_PHOTOMETRIC_DATA);
+	bool bUseIES = LightTypeId == cineware::LIGHT_TYPE_PHOTOMETRIC && MelangeGetBool(InC4DLightPtr, cineware::LIGHT_PHOTOMETRIC_DATA);
 	if (bUseIES)
 	{
-		FString IESFilename = MelangeGetString(InC4DLightPtr, melange::LIGHT_PHOTOMETRIC_FILE);
+		FString IESFilename = MelangeGetString(InC4DLightPtr, cineware::LIGHT_PHOTOMETRIC_FILE);
 
 		IESPath = SearchForFile(IESFilename, C4dDocumentFilename);
 		if (IESPath.IsEmpty())
@@ -1084,15 +1084,15 @@ TSharedPtr<IDatasmithLightActorElement> FDatasmithC4DImporter::ImportLight(melan
 		PointLightActor->SetIntensityUnits(Units);
 
 		// Attenuation radius
-		int32 FalloffOption = MelangeGetInt32(InC4DLightPtr, melange::LIGHT_DETAILS_FALLOFF);
-		if (FalloffOption == melange::LIGHT_DETAILS_FALLOFF_NONE)
+		int32 FalloffOption = MelangeGetInt32(InC4DLightPtr, cineware::LIGHT_DETAILS_FALLOFF);
+		if (FalloffOption == cineware::LIGHT_DETAILS_FALLOFF_NONE)
 		{
 			PointLightActor->SetAttenuationRadius(16384.0f); // Seems to be the maximum value for the slider in the details panel
 		}
 		// TODO: Add support for other falloff types
 		else
 		{
-			PointLightActor->SetAttenuationRadius(MelangeGetFloat(InC4DLightPtr, melange::LIGHT_DETAILS_OUTERDISTANCE));
+			PointLightActor->SetAttenuationRadius(MelangeGetFloat(InC4DLightPtr, cineware::LIGHT_DETAILS_OUTERDISTANCE));
 		}
 	}
 
@@ -1102,11 +1102,11 @@ TSharedPtr<IDatasmithLightActorElement> FDatasmithC4DImporter::ImportLight(melan
 		TSharedPtr<IDatasmithSpotLightElement> SpotLightActor = StaticCastSharedPtr<IDatasmithSpotLightElement>(LightActor);
 
 		// Inner angle
-		float LightInnerAngleInRadians = MelangeGetFloat(InC4DLightPtr, melange::LIGHT_DETAILS_INNERANGLE);
+		float LightInnerAngleInRadians = MelangeGetFloat(InC4DLightPtr, cineware::LIGHT_DETAILS_INNERANGLE);
 		SpotLightActor->SetInnerConeAngle((FMath::RadiansToDegrees(LightInnerAngleInRadians) * 90) / 175);
 
 		// Outer angle
-		float LightOuterAngleInRadians = MelangeGetFloat(InC4DLightPtr, melange::LIGHT_DETAILS_OUTERANGLE);
+		float LightOuterAngleInRadians = MelangeGetFloat(InC4DLightPtr, cineware::LIGHT_DETAILS_OUTERANGLE);
 		SpotLightActor->SetOuterConeAngle((FMath::RadiansToDegrees(LightOuterAngleInRadians) * 90) / 175);
 	}
 
@@ -1116,18 +1116,18 @@ TSharedPtr<IDatasmithLightActorElement> FDatasmithC4DImporter::ImportLight(melan
 		TSharedPtr<IDatasmithAreaLightElement> AreaLightActor = StaticCastSharedPtr<IDatasmithAreaLightElement>(LightActor);
 
 		// Area width
-		AreaLightActor->SetWidth(MelangeGetFloat(InC4DLightPtr, melange::LIGHT_AREADETAILS_SIZEX));
+		AreaLightActor->SetWidth(MelangeGetFloat(InC4DLightPtr, cineware::LIGHT_AREADETAILS_SIZEX));
 
 		// Area length
-		AreaLightActor->SetLength(MelangeGetFloat(InC4DLightPtr, melange::LIGHT_AREADETAILS_SIZEY));
+		AreaLightActor->SetLength(MelangeGetFloat(InC4DLightPtr, cineware::LIGHT_AREADETAILS_SIZEY));
 
 		// Area shape and type
-		EDatasmithLightShape AreaShape = GetDatasmithAreaLightShape(MelangeGetInt32(InC4DLightPtr, melange::LIGHT_AREADETAILS_SHAPE));
+		EDatasmithLightShape AreaShape = GetDatasmithAreaLightShape(MelangeGetInt32(InC4DLightPtr, cineware::LIGHT_AREADETAILS_SHAPE));
 
 		// AreaLightType will default to Point, which is OK for most shapes except the planar shapes like Disc and Rectangle.
 		// Also, if the user enabled the "Z Direction Only" checkbox we'll also use Rect type as the Point type is omnidirectional
 		EDatasmithAreaLightType AreaType = EDatasmithAreaLightType::Point;
-		bool bOnlyZ = MelangeGetBool(InC4DLightPtr, melange::LIGHT_DETAILS_ONLYZ);
+		bool bOnlyZ = MelangeGetBool(InC4DLightPtr, cineware::LIGHT_DETAILS_ONLYZ);
 		if (bOnlyZ || AreaShape == EDatasmithLightShape::Rectangle || AreaShape == EDatasmithLightShape::Disc)
 		{
 			AreaType = EDatasmithAreaLightType::Rect;
@@ -1140,15 +1140,15 @@ TSharedPtr<IDatasmithLightActorElement> FDatasmithC4DImporter::ImportLight(melan
 	return LightActor;
 }
 
-TSharedPtr<IDatasmithCameraActorElement> FDatasmithC4DImporter::ImportCamera(melange::BaseObject* InC4DCameraPtr, const FString& DatasmithName, const FString& DatasmithLabel)
+TSharedPtr<IDatasmithCameraActorElement> FDatasmithC4DImporter::ImportCamera(cineware::BaseObject* InC4DCameraPtr, const FString& DatasmithName, const FString& DatasmithLabel)
 {
 	TSharedPtr<IDatasmithCameraActorElement> CameraActor = FDatasmithSceneFactory::CreateCameraActor(*DatasmithName);
 	CameraActor->SetLabel(*DatasmithLabel);
 
-	melange::GeData C4DData;
+	cineware::GeData C4DData;
 
-	melange::BaseTag* LookAtTag = InC4DCameraPtr->GetTag(Ttargetexpression);
-	melange::BaseList2D* LookAtObject = LookAtTag ? MelangeGetLink(LookAtTag, melange::TARGETEXPRESSIONTAG_LINK) : nullptr;
+	cineware::BaseTag* LookAtTag = InC4DCameraPtr->GetTag(Ttargetexpression);
+	cineware::BaseList2D* LookAtObject = LookAtTag ? MelangeGetLink(LookAtTag, cineware::TARGETEXPRESSIONTAG_LINK) : nullptr;
 	if (LookAtObject)
 	{
 		//LookAtObject can not be a cached object or an instanced object so GetMelangeBaseList2dID should be the final ID
@@ -1162,40 +1162,40 @@ TSharedPtr<IDatasmithCameraActorElement> FDatasmithC4DImporter::ImportCamera(mel
 		NamesOfActorsToKeep.Add(LookAtID.GetValue());
 	}
 
-	float CameraFocusDistanceInCM = MelangeGetFloat(InC4DCameraPtr, melange::CAMERAOBJECT_TARGETDISTANCE);
+	float CameraFocusDistanceInCM = MelangeGetFloat(InC4DCameraPtr, cineware::CAMERAOBJECT_TARGETDISTANCE);
 	CameraActor->SetFocusDistance(CameraFocusDistanceInCM);
 
-	float CameraFocalLengthMilimeters = MelangeGetFloat(InC4DCameraPtr, melange::CAMERA_FOCUS);
+	float CameraFocalLengthMilimeters = MelangeGetFloat(InC4DCameraPtr, cineware::CAMERA_FOCUS);
 	CameraActor->SetFocalLength(CameraFocalLengthMilimeters);
 
-	float CameraHorizontalFieldOfViewInDegree = FMath::RadiansToDegrees(MelangeGetFloat(InC4DCameraPtr, melange::CAMERAOBJECT_FOV));
+	float CameraHorizontalFieldOfViewInDegree = FMath::RadiansToDegrees(MelangeGetFloat(InC4DCameraPtr, cineware::CAMERAOBJECT_FOV));
 	float CameraSensorWidthInMillimeter = 2 * (CameraFocalLengthMilimeters * tan((0.5f * CameraHorizontalFieldOfViewInDegree) / 57.296f));
 	CameraActor->SetSensorWidth(CameraSensorWidthInMillimeter);
 
 	// Set the camera aspect ratio (width/height).
-	melange::RenderData* SceneRenderer = C4dDocument->GetActiveRenderData();
-	melange::Float AspectRatioOfRenderer, RendererWidth, RendererHeight, PixelAspectRatio;
+	cineware::RenderData* SceneRenderer = C4dDocument->GetActiveRenderData();
+	cineware::Float AspectRatioOfRenderer, RendererWidth, RendererHeight, PixelAspectRatio;
 	SceneRenderer->GetResolution(RendererWidth, RendererHeight, PixelAspectRatio, AspectRatioOfRenderer);
 	double AspectRatio = RendererWidth / RendererHeight;
 	CameraActor->SetSensorAspectRatio(static_cast<float>(AspectRatio));
 
 	// We only use manual exposure control with aperture, shutter speed and ISO if the exposure checkbox is enabled
 	// Aperture is always used for depth of field effects though, which is why its outside of this
-	if (MelangeGetBool(InC4DCameraPtr, melange::CAMERAOBJECT_EXPOSURE))
+	if (MelangeGetBool(InC4DCameraPtr, cineware::CAMERAOBJECT_EXPOSURE))
 	{
-		float ShutterSpeed = MelangeGetFloat(InC4DCameraPtr, melange::CAMERAOBJECT_SHUTTER_SPEED_VALUE);
+		float ShutterSpeed = MelangeGetFloat(InC4DCameraPtr, cineware::CAMERAOBJECT_SHUTTER_SPEED_VALUE);
 		CameraActor->GetPostProcess()->SetCameraShutterSpeed(ShutterSpeed ? 1.0f/ShutterSpeed : -1.0f);
 
-		float ISO = MelangeGetFloat(InC4DCameraPtr, melange::CAMERAOBJECT_ISO_VALUE);
+		float ISO = MelangeGetFloat(InC4DCameraPtr, cineware::CAMERAOBJECT_ISO_VALUE);
 		CameraActor->GetPostProcess()->SetCameraISO(ISO ? ISO : -1.0f);
 	}
-	float Aperture = MelangeGetFloat(InC4DCameraPtr, melange::CAMERAOBJECT_FNUMBER_VALUE);
+	float Aperture = MelangeGetFloat(InC4DCameraPtr, cineware::CAMERAOBJECT_FNUMBER_VALUE);
 	CameraActor->SetFStop(Aperture ? Aperture : -1.0f);
 
-	melange::BaseTag* Tag = InC4DCameraPtr->GetFirstTag();
+	cineware::BaseTag* Tag = InC4DCameraPtr->GetFirstTag();
 	while (Tag)
 	{
-		melange::Int32 TagType = Tag->GetType();
+		cineware::Int32 TagType = Tag->GetType();
 		if (TagType == Tcrane)
 		{
 			TSharedRef<FCraneCameraAttributes> Attributes = ExtractCraneCameraAttributes(Tag);
@@ -1230,7 +1230,7 @@ TSharedPtr<IDatasmithTextureElement> FDatasmithC4DImporter::ImportTexture(const
 	return Texture;
 }
 
-FString FDatasmithC4DImporter::GetBaseShaderTextureFilePath(melange::BaseList2D* BaseShader)
+FString FDatasmithC4DImporter::GetBaseShaderTextureFilePath(cineware::BaseList2D* BaseShader)
 {
 	FString TextureFilePath;
 
@@ -1240,12 +1240,12 @@ FString FDatasmithC4DImporter::GetBaseShaderTextureFilePath(melange::BaseList2D*
 		{
 		case Xbitmap:
 		{
-			FString Filepath = MelangeFilenameToPath(static_cast<melange::BaseShader*>(BaseShader)->GetFileName());
+			FString Filepath = MelangeFilenameToPath(static_cast<cineware::BaseShader*>(BaseShader)->GetFileName());
 			TextureFilePath = SearchForFile(Filepath, C4dDocumentFilename);
 			break;
 		}
 		default:
-			TextureFilePath = GetBaseShaderTextureFilePath(static_cast<melange::BaseShader*>(BaseShader)->GetDown());
+			TextureFilePath = GetBaseShaderTextureFilePath(static_cast<cineware::BaseShader*>(BaseShader)->GetDown());
 			break;
 		}
 
@@ -1255,7 +1255,7 @@ FString FDatasmithC4DImporter::GetBaseShaderTextureFilePath(melange::BaseList2D*
 	return TextureFilePath;
 }
 
-TSharedPtr<IDatasmithMaterialInstanceElement> FDatasmithC4DImporter::ImportMaterial(melange::Material* InC4DMaterialPtr)
+TSharedPtr<IDatasmithMaterialInstanceElement> FDatasmithC4DImporter::ImportMaterial(cineware::Material* InC4DMaterialPtr)
 {
 	TOptional<FString> DatasmithName = GetMelangeBaseList2dID(InC4DMaterialPtr);
 	if (!DatasmithName)
@@ -1273,10 +1273,10 @@ TSharedPtr<IDatasmithMaterialInstanceElement> FDatasmithC4DImporter::ImportMater
 	AddBoolToMaterial(MaterialPtr, TEXT("Use_Color"), bUseColor);
 	if (bUseColor)
 	{
-		FLinearColor Color = FLinearColor(MelangeGetLayerColor(InC4DMaterialPtr, melange::MATERIAL_COLOR_COLOR, melange::MATERIAL_COLOR_BRIGHTNESS));
+		FLinearColor Color = FLinearColor(MelangeGetLayerColor(InC4DMaterialPtr, cineware::MATERIAL_COLOR_COLOR, cineware::MATERIAL_COLOR_BRIGHTNESS));
 		AddColorToMaterial(MaterialPtr, TEXT("Color"), Color);
 
-		melange::BaseList2D* MaterialShader = MelangeGetLink(InC4DMaterialPtr, melange::MATERIAL_COLOR_SHADER);
+		cineware::BaseList2D* MaterialShader = MelangeGetLink(InC4DMaterialPtr, cineware::MATERIAL_COLOR_SHADER);
 		FString TextureFilePath = GetBaseShaderTextureFilePath(MaterialShader);
 		TSharedPtr<IDatasmithTextureElement> ColorMap = ImportTexture(TextureFilePath, EDatasmithTextureMode::Diffuse);
 		AddTextureToMaterial(MaterialPtr, TEXT("ColorMap"), ColorMap);
@@ -1288,16 +1288,16 @@ TSharedPtr<IDatasmithMaterialInstanceElement> FDatasmithC4DImporter::ImportMater
 			AddFloatToMaterial(MaterialPtr, TEXT("Exposure"), 0);
 
 			// Check for the good type of Texture Mixing and Blending
-			int32 MixingTypeId = MelangeGetInt32(InC4DMaterialPtr, melange::MATERIAL_COLOR_TEXTUREMIXING);
+			int32 MixingTypeId = MelangeGetInt32(InC4DMaterialPtr, cineware::MATERIAL_COLOR_TEXTUREMIXING);
 			switch (MixingTypeId)
 			{
-			case melange::MATERIAL_TEXTUREMIXING_ADD:
+			case cineware::MATERIAL_TEXTUREMIXING_ADD:
 				AddBoolToMaterial(MaterialPtr, TEXT("MixMode_Add"), true);
 				break;
-			case melange::MATERIAL_TEXTUREMIXING_SUBTRACT:
+			case cineware::MATERIAL_TEXTUREMIXING_SUBTRACT:
 				AddBoolToMaterial(MaterialPtr, TEXT("MixMode_Subtract"), true);
 				break;
-			case melange::MATERIAL_TEXTUREMIXING_MULTIPLY:
+			case cineware::MATERIAL_TEXTUREMIXING_MULTIPLY:
 				AddBoolToMaterial(MaterialPtr, TEXT("MixMode_Multiply"), true);
 				break;
 			default: // MATERIAL_TEXTUREMIXING_NORMAL
@@ -1305,7 +1305,7 @@ TSharedPtr<IDatasmithMaterialInstanceElement> FDatasmithC4DImporter::ImportMater
 				break;
 			}
 
-			float MixStrength = MelangeGetFloat(InC4DMaterialPtr, melange::MATERIAL_COLOR_TEXTURESTRENGTH);
+			float MixStrength = MelangeGetFloat(InC4DMaterialPtr, cineware::MATERIAL_COLOR_TEXTURESTRENGTH);
 			AddFloatToMaterial(MaterialPtr, TEXT("Mix_Strength"), MixStrength);
 		}
 	}
@@ -1315,13 +1315,13 @@ TSharedPtr<IDatasmithMaterialInstanceElement> FDatasmithC4DImporter::ImportMater
 	AddBoolToMaterial(MaterialPtr, TEXT("Use_Emissive"), bUseEmissive);
 	if (bUseEmissive)
 	{
-		float EmissiveGlowStrength = MelangeGetFloat(InC4DMaterialPtr, melange::MATERIAL_LUMINANCE_BRIGHTNESS);
+		float EmissiveGlowStrength = MelangeGetFloat(InC4DMaterialPtr, cineware::MATERIAL_LUMINANCE_BRIGHTNESS);
 		AddFloatToMaterial(MaterialPtr, TEXT("Emissive_Glow_Strength"), EmissiveGlowStrength);
 
-		FLinearColor EmissiveColor = FLinearColor(MelangeGetColor(InC4DMaterialPtr, melange::MATERIAL_LUMINANCE_COLOR));
+		FLinearColor EmissiveColor = FLinearColor(MelangeGetColor(InC4DMaterialPtr, cineware::MATERIAL_LUMINANCE_COLOR));
 		AddColorToMaterial(MaterialPtr, TEXT("Emissive_Color"), EmissiveColor);
 
-		melange::BaseList2D* LuminanceShader = MelangeGetLink(InC4DMaterialPtr, melange::MATERIAL_LUMINANCE_SHADER);
+		cineware::BaseList2D* LuminanceShader = MelangeGetLink(InC4DMaterialPtr, cineware::MATERIAL_LUMINANCE_SHADER);
 		FString LuminanceFilePath = GetBaseShaderTextureFilePath(LuminanceShader);
 		TSharedPtr<IDatasmithTextureElement> EmissiveMap = ImportTexture(LuminanceFilePath, EDatasmithTextureMode::Other);
 		AddTextureToMaterial(MaterialPtr, TEXT("Emissive_Map"), EmissiveMap);
@@ -1330,7 +1330,7 @@ TSharedPtr<IDatasmithMaterialInstanceElement> FDatasmithC4DImporter::ImportMater
 		AddBoolToMaterial(MaterialPtr, TEXT("Use_EmissiveMap"), bUseEmissiveMap);
 		if (bUseEmissiveMap)
 		{
-			float EmissiveMapExposure = MelangeGetFloat(InC4DMaterialPtr, melange::MATERIAL_LUMINANCE_TEXTURESTRENGTH);
+			float EmissiveMapExposure = MelangeGetFloat(InC4DMaterialPtr, cineware::MATERIAL_LUMINANCE_TEXTURESTRENGTH);
 			AddFloatToMaterial(MaterialPtr, TEXT("Emissive_Map_Exposure"), EmissiveMapExposure);
 		}
 	}
@@ -1342,7 +1342,7 @@ TSharedPtr<IDatasmithMaterialInstanceElement> FDatasmithC4DImporter::ImportMater
 	{
 		MaterialPtr->SetMaterialType(EDatasmithReferenceMaterialType::Transparent);
 
-		melange::BaseList2D* TransparencyShader = MelangeGetLink(InC4DMaterialPtr, melange::MATERIAL_TRANSPARENCY_SHADER);
+		cineware::BaseList2D* TransparencyShader = MelangeGetLink(InC4DMaterialPtr, cineware::MATERIAL_TRANSPARENCY_SHADER);
 		FString TransparencyMapPath = GetBaseShaderTextureFilePath(TransparencyShader);
 		TSharedPtr<IDatasmithTextureElement> TransparencyMap = ImportTexture(TransparencyMapPath, EDatasmithTextureMode::Other);
 		AddTextureToMaterial(MaterialPtr, TEXT("Transparency_Map"), TransparencyMap);
@@ -1351,30 +1351,30 @@ TSharedPtr<IDatasmithMaterialInstanceElement> FDatasmithC4DImporter::ImportMater
 		AddBoolToMaterial(MaterialPtr, TEXT("Use_TransparencyMap"), bUseTransparencyMap);
 		if (bUseTransparencyMap)
 		{
-			float TextureStrength = MelangeGetFloat(InC4DMaterialPtr, melange::MATERIAL_TRANSPARENCY_TEXTURESTRENGTH);
+			float TextureStrength = MelangeGetFloat(InC4DMaterialPtr, cineware::MATERIAL_TRANSPARENCY_TEXTURESTRENGTH);
 			AddFloatToMaterial(MaterialPtr, TEXT("TransparencyMap_Amount"), TextureStrength);
 		}
 		else
 		{
-			float BrightnessValue = MelangeGetFloat(InC4DMaterialPtr, melange::MATERIAL_TRANSPARENCY_BRIGHTNESS);
-			FVector3f TransparencyColor = (FVector3f)MelangeGetVector(InC4DMaterialPtr, melange::MATERIAL_TRANSPARENCY_COLOR);
+			float BrightnessValue = MelangeGetFloat(InC4DMaterialPtr, cineware::MATERIAL_TRANSPARENCY_BRIGHTNESS);
+			FVector3f TransparencyColor = (FVector3f)MelangeGetVector(InC4DMaterialPtr, cineware::MATERIAL_TRANSPARENCY_COLOR);
 
 			// In Cinema4D Transparency Color seems to be used just as another multiplier for the opacity, not as an actual color
 			AddFloatToMaterial(MaterialPtr, TEXT("Transparency_Amount"), BrightnessValue * TransparencyColor.X * TransparencyColor.Y * TransparencyColor.Z);
 		}
 
-		float TransparencyRefraction = MelangeGetFloat(InC4DMaterialPtr, melange::MATERIAL_TRANSPARENCY_REFRACTION);
+		float TransparencyRefraction = MelangeGetFloat(InC4DMaterialPtr, cineware::MATERIAL_TRANSPARENCY_REFRACTION);
 		AddFloatToMaterial(MaterialPtr, TEXT("Transparency_Refraction"), TransparencyRefraction);
 	}
 
-	melange::GeData C4DData;
+	cineware::GeData C4DData;
 
 	// Specular
 	bool bUseSpecular = InC4DMaterialPtr->GetChannelState(CHANNEL_REFLECTION);
 	AddBoolToMaterial(MaterialPtr, TEXT("Use_Specular"), bUseSpecular);
 	if (bUseSpecular)
 	{
-		melange::Int32 ReflectionLayerCount = InC4DMaterialPtr->GetReflectionLayerCount();
+		cineware::Int32 ReflectionLayerCount = InC4DMaterialPtr->GetReflectionLayerCount();
 		if (ReflectionLayerCount > 0)
 		{
 			bool bUseReflectionColor = false;
@@ -1383,14 +1383,14 @@ TSharedPtr<IDatasmithMaterialInstanceElement> FDatasmithC4DImporter::ImportMater
 			FVector ReflectionColor(0, 0, 0);
 			for (int32 LayerIndex = ReflectionLayerCount - 1; LayerIndex >= 0; --LayerIndex)
 			{
-				melange::ReflectionLayer* ReflectionLayer = InC4DMaterialPtr->GetReflectionLayerIndex(LayerIndex);
+				cineware::ReflectionLayer* ReflectionLayer = InC4DMaterialPtr->GetReflectionLayerIndex(LayerIndex);
 				if (ReflectionLayer == nullptr)
 				{
 					continue;
 				}
 
-				melange::Int32 ReflectionLayerBaseId = ReflectionLayer->GetDataID();
-				melange::Int32 ReflectionLayerFlags = ReflectionLayer->flags;
+				cineware::Int32 ReflectionLayerBaseId = ReflectionLayer->GetDataID();
+				cineware::Int32 ReflectionLayerFlags = ReflectionLayer->GetFlags();
 
 				// Don't fetch colors from reflectance layers that, regardless of fresnel function, don't seem to contribute a lot to main base color
 				int32 LayerType = MelangeGetInt32(InC4DMaterialPtr, ReflectionLayerBaseId+REFLECTION_LAYER_MAIN_DISTRIBUTION);
@@ -1414,7 +1414,7 @@ TSharedPtr<IDatasmithMaterialInstanceElement> FDatasmithC4DImporter::ImportMater
 
 					// This is a temporary solution in order to let some color from reflectance layers factor in to the final basecolor depending on their
 					// fresnel function
-					melange::Int32 FresnelMode = InC4DMaterialPtr->GetParameter(ReflectionLayerBaseId + REFLECTION_LAYER_FRESNEL_MODE, C4DData) ? C4DData.GetInt32() : REFLECTION_FRESNEL_NONE;
+					cineware::Int32 FresnelMode = InC4DMaterialPtr->GetParameter(ReflectionLayerBaseId + REFLECTION_LAYER_FRESNEL_MODE, C4DData) ? C4DData.GetInt32() : REFLECTION_FRESNEL_NONE;
 					switch (FresnelMode)
 					{
 					case REFLECTION_FRESNEL_NONE:
@@ -1459,18 +1459,18 @@ TSharedPtr<IDatasmithMaterialInstanceElement> FDatasmithC4DImporter::ImportMater
 			}
 
 			// Only set those one for the last Layer of reflection
-			melange::ReflectionLayer* ReflectionLayer = InC4DMaterialPtr->GetReflectionLayerIndex(0);
+			cineware::ReflectionLayer* ReflectionLayer = InC4DMaterialPtr->GetReflectionLayerIndex(0);
 
 			bool bUseReflectance = (ReflectionLayer != nullptr);
 			AddBoolToMaterial(MaterialPtr, TEXT("Use_Reflectance"), bUseReflectance);
 			if (bUseReflectance)
 			{
-				melange::Int32 ReflectionLayerBaseId = ReflectionLayer->GetDataID();
+				cineware::Int32 ReflectionLayerBaseId = ReflectionLayer->GetDataID();
 
 				float SpecularStrength = MelangeGetFloat(InC4DMaterialPtr, ReflectionLayerBaseId + REFLECTION_LAYER_MAIN_VALUE_SPECULAR);
 				AddFloatToMaterial(MaterialPtr, TEXT("Specular_Strength"), SpecularStrength);
 
-				melange::BaseList2D* RoughnessShader = MelangeGetLink(InC4DMaterialPtr, ReflectionLayerBaseId + REFLECTION_LAYER_MAIN_SHADER_ROUGHNESS);
+				cineware::BaseList2D* RoughnessShader = MelangeGetLink(InC4DMaterialPtr, ReflectionLayerBaseId + REFLECTION_LAYER_MAIN_SHADER_ROUGHNESS);
 				FString RoughnessMapPath = GetBaseShaderTextureFilePath(RoughnessShader);
 				TSharedPtr<IDatasmithTextureElement> RoughnessMap1 = ImportTexture(RoughnessMapPath, EDatasmithTextureMode::Diffuse);
 				AddTextureToMaterial(MaterialPtr, TEXT("RoughnessMap1"), RoughnessMap1);
@@ -1484,7 +1484,7 @@ TSharedPtr<IDatasmithMaterialInstanceElement> FDatasmithC4DImporter::ImportMater
 				}
 				else
 				{
-					float RoughnessStrength = MelangeGetFloat(InC4DMaterialPtr, melange::MATERIAL_SPECULAR_WIDTH /*appears to be the computed roughness*/);
+					float RoughnessStrength = MelangeGetFloat(InC4DMaterialPtr, cineware::MATERIAL_SPECULAR_WIDTH /*appears to be the computed roughness*/);
 					AddFloatToMaterial(MaterialPtr, TEXT("Roughness_Strength"), RoughnessStrength);
 				}
 
@@ -1496,7 +1496,7 @@ TSharedPtr<IDatasmithMaterialInstanceElement> FDatasmithC4DImporter::ImportMater
 				{
 					AddFloatToMaterial(MaterialPtr, TEXT("Metalic_Amount"), 0.5f);
 
-					melange::BaseList2D* MetallicShader = MelangeGetLink(InC4DMaterialPtr, ReflectionLayerBaseId + REFLECTION_LAYER_TRANS_TEXTURE);
+					cineware::BaseList2D* MetallicShader = MelangeGetLink(InC4DMaterialPtr, ReflectionLayerBaseId + REFLECTION_LAYER_TRANS_TEXTURE);
 					FString MetallicMapPath = GetBaseShaderTextureFilePath(MetallicShader);
 					TSharedPtr<IDatasmithTextureElement> MetalicMap = ImportTexture(MetallicMapPath, EDatasmithTextureMode::Specular);
 					AddTextureToMaterial(MaterialPtr, TEXT("MetalicMap"), MetalicMap);
@@ -1513,14 +1513,14 @@ TSharedPtr<IDatasmithMaterialInstanceElement> FDatasmithC4DImporter::ImportMater
 	AddBoolToMaterial(MaterialPtr, TEXT("Use_AO"), bUseAO);
 	if (bUseAO)
 	{
-		melange::BaseList2D* DiffusionShader = MelangeGetLink(InC4DMaterialPtr, melange::MATERIAL_DIFFUSION_SHADER);
+		cineware::BaseList2D* DiffusionShader = MelangeGetLink(InC4DMaterialPtr, cineware::MATERIAL_DIFFUSION_SHADER);
 		FString AOMapPath = GetBaseShaderTextureFilePath(DiffusionShader);
 		TSharedPtr<IDatasmithTextureElement> AOMap = ImportTexture(AOMapPath, EDatasmithTextureMode::Diffuse);
 		AddTextureToMaterial(MaterialPtr, TEXT("AO_Map"), AOMap);
 
 		if (!AOMapPath.IsEmpty())
 		{
-			float AOStrength = MelangeGetFloat(InC4DMaterialPtr, melange::MATERIAL_DIFFUSION_TEXTURESTRENGTH);
+			float AOStrength = MelangeGetFloat(InC4DMaterialPtr, cineware::MATERIAL_DIFFUSION_TEXTURESTRENGTH);
 			AddFloatToMaterial(MaterialPtr, TEXT("AO_Strength"), AOStrength);
 		}
 	}
@@ -1532,12 +1532,12 @@ TSharedPtr<IDatasmithMaterialInstanceElement> FDatasmithC4DImporter::ImportMater
 	{
 		MaterialPtr->SetMaterialType(EDatasmithReferenceMaterialType::CutOut);
 
-		melange::BaseList2D* AlphaShader = MelangeGetLink(InC4DMaterialPtr, melange::MATERIAL_ALPHA_SHADER);
+		cineware::BaseList2D* AlphaShader = MelangeGetLink(InC4DMaterialPtr, cineware::MATERIAL_ALPHA_SHADER);
 		FString AlphaMapPath = GetBaseShaderTextureFilePath(AlphaShader);
 		TSharedPtr<IDatasmithTextureElement> AlphaMap = ImportTexture(AlphaMapPath, EDatasmithTextureMode::Diffuse);
 		AddTextureToMaterial(MaterialPtr, TEXT("Alpha_Map"), AlphaMap);
 
-		bool bUseAlphaInvert = MelangeGetBool(InC4DMaterialPtr, melange::MATERIAL_ALPHA_INVERT);
+		bool bUseAlphaInvert = MelangeGetBool(InC4DMaterialPtr, cineware::MATERIAL_ALPHA_INVERT);
 		AddBoolToMaterial(MaterialPtr, TEXT("Use_Alpha_Invert"), bUseAlphaInvert);
 	}
 
@@ -1546,14 +1546,14 @@ TSharedPtr<IDatasmithMaterialInstanceElement> FDatasmithC4DImporter::ImportMater
 	AddBoolToMaterial(MaterialPtr, TEXT("Use_Normal"), bUseNormal);
 	if (bUseNormal)
 	{
-		melange::BaseList2D* NormalShader = MelangeGetLink(InC4DMaterialPtr, melange::MATERIAL_NORMAL_SHADER);
+		cineware::BaseList2D* NormalShader = MelangeGetLink(InC4DMaterialPtr, cineware::MATERIAL_NORMAL_SHADER);
 		FString NormalMapPath = GetBaseShaderTextureFilePath(NormalShader);
 		TSharedPtr<IDatasmithTextureElement> NormalMap = ImportTexture(NormalMapPath, EDatasmithTextureMode::Normal);
 		AddTextureToMaterial(MaterialPtr, TEXT("Normal_Map"), NormalMap);
 
 		if (!NormalMapPath.IsEmpty())
 		{
-			float NormalMapStrength = MelangeGetFloat(InC4DMaterialPtr, melange::MATERIAL_NORMAL_STRENGTH);
+			float NormalMapStrength = MelangeGetFloat(InC4DMaterialPtr, cineware::MATERIAL_NORMAL_STRENGTH);
 			AddFloatToMaterial(MaterialPtr, TEXT("Normal_Strength"), NormalMapStrength);
 		}
 	}
@@ -1562,7 +1562,7 @@ TSharedPtr<IDatasmithMaterialInstanceElement> FDatasmithC4DImporter::ImportMater
 	return MaterialPtr;
 }
 
-TSharedPtr<IDatasmithMaterialInstanceElement> FDatasmithC4DImporter::ImportSimpleColorMaterial(melange::BaseObject* Object, int32 UseColor)
+TSharedPtr<IDatasmithMaterialInstanceElement> FDatasmithC4DImporter::ImportSimpleColorMaterial(cineware::BaseObject* Object, int32 UseColor)
 {
 	TOptional<FString> DatasmithName = GetMelangeBaseList2dID(Object);
 	if (!DatasmithName)
@@ -1572,15 +1572,15 @@ TSharedPtr<IDatasmithMaterialInstanceElement> FDatasmithC4DImporter::ImportSimpl
 	FString DatasmithLabel = FDatasmithUtils::SanitizeObjectName(MelangeObjectName(Object) + TEXT("_DisplayColor"));
 
 	FVector DisplayColor{1.0f, 1.0f, 1.0f};
-	if (UseColor == melange::ID_BASEOBJECT_USECOLOR_AUTOMATIC || UseColor == melange::ID_BASEOBJECT_USECOLOR_ALWAYS)
+	if (UseColor == cineware::ID_BASEOBJECT_USECOLOR_AUTOMATIC || UseColor == cineware::ID_BASEOBJECT_USECOLOR_ALWAYS)
 	{
-		DisplayColor = MelangeGetColor(Object, melange::ID_BASEOBJECT_COLOR);
+		DisplayColor = MelangeGetColor(Object, cineware::ID_BASEOBJECT_COLOR);
 	}
-	else if (UseColor == melange::ID_BASEOBJECT_USECOLOR_LAYER)
+	else if (UseColor == cineware::ID_BASEOBJECT_USECOLOR_LAYER)
 	{
-		if (melange::LayerObject* LayerObject = static_cast<melange::LayerObject*>(MelangeGetLink(Object, melange::ID_LAYER_LINK)))
+		if (cineware::LayerObject* LayerObject = static_cast<cineware::LayerObject*>(MelangeGetLink(Object, cineware::ID_LAYER_LINK)))
 		{
-			DisplayColor = MelangeGetColor(LayerObject, melange::ID_LAYER_COLOR);
+			DisplayColor = MelangeGetColor(LayerObject, cineware::ID_LAYER_COLOR);
 		}
 		else
 		{
@@ -1609,7 +1609,7 @@ TSharedPtr<IDatasmithMaterialInstanceElement> FDatasmithC4DImporter::ImportSimpl
 	return Material;
 }
 
-bool FDatasmithC4DImporter::ImportMaterialHierarchy(melange::BaseMaterial* InC4DMaterialPtr)
+bool FDatasmithC4DImporter::ImportMaterialHierarchy(cineware::BaseMaterial* InC4DMaterialPtr)
 {
 	// Reinitialize the scene material map and texture set.
 	MaterialNameToMaterialElement.Empty();
@@ -1618,7 +1618,7 @@ bool FDatasmithC4DImporter::ImportMaterialHierarchy(melange::BaseMaterial* InC4D
 	{
 		if (InC4DMaterialPtr->GetType() == Mmaterial)
 		{
-			if (TSharedPtr<IDatasmithMaterialInstanceElement> DatasmithMaterial = ImportMaterial(static_cast<melange::Material*>(InC4DMaterialPtr)))
+			if (TSharedPtr<IDatasmithMaterialInstanceElement> DatasmithMaterial = ImportMaterial(static_cast<cineware::Material*>(InC4DMaterialPtr)))
 			{
 				MaterialNameToMaterialElement.Add(DatasmithMaterial->GetName(), DatasmithMaterial);
 			}
@@ -1632,7 +1632,7 @@ bool FDatasmithC4DImporter::ImportMaterialHierarchy(melange::BaseMaterial* InC4D
 	return true;
 }
 
-FString FDatasmithC4DImporter::CustomizeMaterial(const FString& InMaterialID, const FString& InMeshID, melange::TextureTag* InTextureTag)
+FString FDatasmithC4DImporter::CustomizeMaterial(const FString& InMaterialID, const FString& InMeshID, cineware::TextureTag* InTextureTag)
 {
 	FString CustomMaterialID = InMaterialID + InMeshID;
 
@@ -1643,11 +1643,11 @@ FString FDatasmithC4DImporter::CustomizeMaterial(const FString& InMaterialID, co
 
 	if (MaterialNameToMaterialElement.Contains(InMaterialID))
 	{
-		melange::GeData Data;
-		float OffsetX = MelangeGetFloat(InTextureTag, melange::TEXTURETAG_OFFSETX);
-		float OffsetY = MelangeGetFloat(InTextureTag, melange::TEXTURETAG_OFFSETY);
-		float TilesX = MelangeGetFloat(InTextureTag, melange::TEXTURETAG_TILESX);
-		float TilesY = MelangeGetFloat(InTextureTag, melange::TEXTURETAG_TILESY);
+		cineware::GeData Data;
+		float OffsetX = MelangeGetFloat(InTextureTag, cineware::TEXTURETAG_OFFSETX);
+		float OffsetY = MelangeGetFloat(InTextureTag, cineware::TEXTURETAG_OFFSETY);
+		float TilesX = MelangeGetFloat(InTextureTag, cineware::TEXTURETAG_TILESX);
+		float TilesY = MelangeGetFloat(InTextureTag, cineware::TEXTURETAG_TILESY);
 
 		if (OffsetX != 0.0F || OffsetY != 0.0F || TilesX != 1.0F || TilesY != 1.0F)
 		{
@@ -1676,7 +1676,7 @@ FString FDatasmithC4DImporter::CustomizeMaterial(const FString& InMaterialID, co
 	return InMaterialID;
 }
 
-TMap<int32, FString> FDatasmithC4DImporter::GetCustomizedMaterialAssignment(const FString& DatasmithMeshName, const TArray<melange::TextureTag*>& TextureTags)
+TMap<int32, FString> FDatasmithC4DImporter::GetCustomizedMaterialAssignment(const FString& DatasmithMeshName, const TArray<cineware::TextureTag*>& TextureTags)
 {
 	TMap<int32, FString> SlotToMaterialName;
 
@@ -1684,10 +1684,10 @@ TMap<int32, FString> FDatasmithC4DImporter::GetCustomizedMaterialAssignment(cons
 	// actually represents a material "instance", and might have different settings like texture tiling
 	for (int32 SlotIndex = 0; SlotIndex < TextureTags.Num(); ++SlotIndex)
 	{
-		melange::TextureTag* Tag = TextureTags[SlotIndex];
+		cineware::TextureTag* Tag = TextureTags[SlotIndex];
 
 		FString CustomizedMaterialName;
-		melange::BaseList2D* TextureMaterial = Tag ? MelangeGetLink(Tag, melange::TEXTURETAG_MATERIAL) : nullptr;
+		cineware::BaseList2D* TextureMaterial = Tag ? MelangeGetLink(Tag, cineware::TEXTURETAG_MATERIAL) : nullptr;
 		if (TextureMaterial)
 		{
 			// This can also return an existing material without necessarily spawning a new instance
@@ -1703,7 +1703,7 @@ TMap<int32, FString> FDatasmithC4DImporter::GetCustomizedMaterialAssignment(cons
 	return SlotToMaterialName;
 }
 
-TSharedPtr<IDatasmithMeshActorElement> FDatasmithC4DImporter::ImportPolygon(melange::PolygonObject* PolyObject, const FString& DatasmithActorName, const FString& DatasmithActorLabel, const TArray<melange::TextureTag*>& TextureTags)
+TSharedPtr<IDatasmithMeshActorElement> FDatasmithC4DImporter::ImportPolygon(cineware::PolygonObject* PolyObject, const FString& DatasmithActorName, const FString& DatasmithActorLabel, const TArray<cineware::TextureTag*>& TextureTags)
 {
 	FMD5Hash PolygonHash = ComputePolygonDataHash(PolyObject);
 	FString HashString = BytesToHex(PolygonHash.GetBytes(), PolygonHash.GetSize());
@@ -1754,7 +1754,7 @@ TSharedPtr<IDatasmithMeshActorElement> FDatasmithC4DImporter::ImportPolygon(mela
 		}
 	}
 
-	const int32 UseColor = MelangeGetInt32(PolyObject, melange::ID_BASEOBJECT_USECOLOR);
+	const int32 UseColor = MelangeGetInt32(PolyObject, cineware::ID_BASEOBJECT_USECOLOR);
 
 	// Add material overrides
 	TMap<int32, FString> SlotIndexToMaterialName = GetCustomizedMaterialAssignment(DatasmithMeshName, TextureTags);
@@ -1765,11 +1765,11 @@ TSharedPtr<IDatasmithMeshActorElement> FDatasmithC4DImporter::ImportPolygon(mela
 
 		// Pick whether we use the display color material or a texturetag material
 		TSharedPtr<IDatasmithMaterialInstanceElement> TargetMaterial = nullptr;
-		if (UseColor == melange::ID_BASEOBJECT_USECOLOR_ALWAYS || UseColor == melange::ID_BASEOBJECT_USECOLOR_LAYER)
+		if (UseColor == cineware::ID_BASEOBJECT_USECOLOR_ALWAYS || UseColor == cineware::ID_BASEOBJECT_USECOLOR_LAYER)
 		{
 			TargetMaterial = ImportSimpleColorMaterial(PolyObject, UseColor);
 		}
-		else if (UseColor == melange::ID_BASEOBJECT_USECOLOR_AUTOMATIC)
+		else if (UseColor == cineware::ID_BASEOBJECT_USECOLOR_AUTOMATIC)
 		{
 			if (MaterialName.IsEmpty())
 			{
@@ -1807,7 +1807,7 @@ TSharedPtr<IDatasmithMeshActorElement> FDatasmithC4DImporter::ImportPolygon(mela
 	return MeshActorElement;
 }
 
-void MarkActorsAsParticlesRecursive(melange::BaseObject* ActorObject, TSet<melange::BaseObject*>& ParticleActors)
+void MarkActorsAsParticlesRecursive(cineware::BaseObject* ActorObject, TSet<cineware::BaseObject*>& ParticleActors)
 {
 	if (ActorObject == nullptr)
 	{
@@ -1820,7 +1820,7 @@ void MarkActorsAsParticlesRecursive(melange::BaseObject* ActorObject, TSet<melan
 	MarkActorsAsParticlesRecursive(ActorObject->GetNext(), ParticleActors);
 }
 
-void FDatasmithC4DImporter::MarkActorsAsParticles(melange::BaseObject* EmitterObject, melange::BaseObject* EmittersCache)
+void FDatasmithC4DImporter::MarkActorsAsParticles(cineware::BaseObject* EmitterObject, cineware::BaseObject* EmittersCache)
 {
 	if (EmitterObject == nullptr || EmittersCache == nullptr)
 	{
@@ -1828,8 +1828,8 @@ void FDatasmithC4DImporter::MarkActorsAsParticles(melange::BaseObject* EmitterOb
 	}
 
 	// C4D only emits mesh "particles" if this "Show Objects" checkbox is checked. Else it just emits actual particles
-	melange::GeData Data;
-	if (EmitterObject->GetParameter(melange::PARTICLEOBJECT_SHOWOBJECTS, Data) && Data.GetType() == melange::DA_LONG && Data.GetBool())
+	cineware::GeData Data;
+	if (EmitterObject->GetParameter(cineware::PARTICLEOBJECT_SHOWOBJECTS, Data) && Data.GetType() == cineware::DA_LONG && Data.GetBool())
 	{
 		MarkActorsAsParticlesRecursive(EmittersCache->GetDown(), ParticleActors);
 	}
@@ -1837,15 +1837,15 @@ void FDatasmithC4DImporter::MarkActorsAsParticles(melange::BaseObject* EmitterOb
 
 namespace
 {
-	melange::Float MelangeFPS = 0;
+	cineware::Float MelangeFPS = 0;
 
 	void AddFrameValueToAnimMap(
-		melange::BaseObject* Object,
+		cineware::BaseObject* Object,
 		int32 FrameNumber,
 		int32 TransformVectorIndex,
 		EDatasmithTransformType TransformType,
-		melange::Float FrameValue,
-		melange::Int32 MelangeTransformType,
+		cineware::Float FrameValue,
+		cineware::Int32 MelangeTransformType,
 		FVector& InitialSize,
 		TMap<int32, TMap<EDatasmithTransformType, FVector>>& TransformFrames,
 		TMap<EDatasmithTransformType, FVector> InitialValues
@@ -1885,11 +1885,11 @@ namespace
 
 void FDatasmithC4DImporter::ImportAnimations(TSharedPtr<IDatasmithActorElement> ActorElement)
 {
-	melange::BaseObject* Object = *ActorElementToAnimationSources.Find(ActorElement.Get());
-	melange::Int32 ObjectType = Object->GetType();
+	cineware::BaseObject* Object = *ActorElementToAnimationSources.Find(ActorElement.Get());
+	cineware::Int32 ObjectType = Object->GetType();
 
 	TMap<EDatasmithTransformType, FVector> InitialValues;
-	melange::Vector MelangeRotation = Object->GetRelRot();
+	cineware::Vector MelangeRotation = Object->GetRelRot();
 	InitialValues.Add(EDatasmithTransformType::Rotation) = FVector(static_cast<float>(MelangeRotation.x), static_cast<float>(MelangeRotation.y), static_cast<float>(MelangeRotation.z));
 	InitialValues.Add(EDatasmithTransformType::Translation) = MelangeVectorToFVector(Object->GetRelPos());
 	InitialValues.Add(EDatasmithTransformType::Scale) = MelangeVectorToFVector(Object->GetRelScale());
@@ -1899,14 +1899,14 @@ void FDatasmithC4DImporter::ImportAnimations(TSharedPtr<IDatasmithActorElement>
 
 	// If we have AlignToSpline animations, the splines are stored with their points in world space,
 	// so we must move them into the object's local space
-	melange::Matrix WorldToLocal = ~Object->GetUpMg();
+	cineware::Matrix WorldToLocal = ~Object->GetUpMg();
 
-	melange::ROTATIONORDER RotationOrder = Object->GetRotationOrder();
+	cineware::ROTATIONORDER RotationOrder = Object->GetRotationOrder();
 
 	// Import animations on the object's tags
-	for(melange::BaseTag* Tag = Object->GetFirstTag(); Tag; Tag=Tag->GetNext())
+	for(cineware::BaseTag* Tag = Object->GetFirstTag(); Tag; Tag=Tag->GetNext())
 	{
-		melange::Int32 TagType = Tag->GetType();
+		cineware::Int32 TagType = Tag->GetType();
 
 		if (TagType == Tcrane && ObjectType == Ocamera)
 		{
@@ -1917,15 +1917,15 @@ void FDatasmithC4DImporter::ImportAnimations(TSharedPtr<IDatasmithActorElement>
 				continue;
 			}
 
-			TMap<int32, melange::CCurve*> CurvesByAttribute;
+			TMap<int32, cineware::CCurve*> CurvesByAttribute;
 
-			melange::BaseTime MinStartTime(FLT_MAX);
-			melange::BaseTime MaxEndTime(-FLT_MAX);
+			cineware::BaseTime MinStartTime(FLT_MAX);
+			cineware::BaseTime MaxEndTime(-FLT_MAX);
 
 			// Get tracks for all animated properties
-			for (melange::CTrack* Track = Tag->GetFirstCTrack(); Track; Track = Track->GetNext())
+			for (cineware::CTrack* Track = Tag->GetFirstCTrack(); Track; Track = Track->GetNext())
 			{
-				melange::DescID TrackDescID = Track->GetDescriptionID();
+				cineware::DescID TrackDescID = Track->GetDescriptionID();
 				int32 Depth = TrackDescID.GetDepth();
 				if (Depth != 1)
 				{
@@ -1933,7 +1933,7 @@ void FDatasmithC4DImporter::ImportAnimations(TSharedPtr<IDatasmithActorElement>
 				}
 				int32 AttributeID = TrackDescID[0].id;
 
-				melange::CCurve* Curve = Track->GetCurve();
+				cineware::CCurve* Curve = Track->GetCurve();
 				if (!Curve || Curve->GetKeyCount() == 0)
 				{
 					continue;
@@ -1954,13 +1954,13 @@ void FDatasmithC4DImporter::ImportAnimations(TSharedPtr<IDatasmithActorElement>
 			int32 LastFrame = MaxEndTime.GetFrame(MelangeFPS);
 			for (int32 FrameNumber = FirstFrame; FrameNumber <= LastFrame; ++FrameNumber)
 			{
-				melange::BaseTime FrameTime = melange::BaseTime(MinStartTime.Get() + melange::Float((FrameNumber-FirstFrame) * (1.0 / MelangeFPS)));
+				cineware::BaseTime FrameTime = cineware::BaseTime(MinStartTime.Get() + cineware::Float((FrameNumber-FirstFrame) * (1.0 / MelangeFPS)));
 
 				// Construct the FCraneCameraAttributes struct for this frame
 				FCraneCameraAttributes AttributesForFrame = **FoundAttributes;
-				for (const TPair<int, melange::CCurve*>& Pair : CurvesByAttribute)
+				for (const TPair<int, cineware::CCurve*>& Pair : CurvesByAttribute)
 				{
-					const melange::CCurve* AttributeCurve = Pair.Value;
+					const cineware::CCurve* AttributeCurve = Pair.Value;
 					const int32& AttributeID = Pair.Key;
 					double AttributeValue = AttributeCurve->GetValue(FrameTime);
 
@@ -1986,7 +1986,7 @@ void FDatasmithC4DImporter::ImportAnimations(TSharedPtr<IDatasmithActorElement>
 						Component,
 						EDatasmithTransformType::Translation,
 						Translation[Component],
-						melange::ID_BASEOBJECT_REL_POSITION,
+						cineware::ID_BASEOBJECT_REL_POSITION,
 						InitialSize,
 						TransformFrames,
 						InitialValues);
@@ -1997,7 +1997,7 @@ void FDatasmithC4DImporter::ImportAnimations(TSharedPtr<IDatasmithActorElement>
 						Component,
 						EDatasmithTransformType::Rotation,
 						FMath::DegreesToRadians(RotationEuler[Component]),
-						melange::ID_BASEOBJECT_REL_ROTATION,
+						cineware::ID_BASEOBJECT_REL_ROTATION,
 						InitialSize,
 						TransformFrames,
 						InitialValues);
@@ -2008,7 +2008,7 @@ void FDatasmithC4DImporter::ImportAnimations(TSharedPtr<IDatasmithActorElement>
 		// base moves along the spline. We don't support that for now
 		else if (TagType == Taligntospline)
 		{
-			melange::SplineObject* SplineObj = static_cast<melange::SplineObject*>(MelangeGetLink(Tag, melange::ALIGNTOSPLINETAG_LINK));
+			cineware::SplineObject* SplineObj = static_cast<cineware::SplineObject*>(MelangeGetLink(Tag, cineware::ALIGNTOSPLINETAG_LINK));
 			if (!SplineObj)
 			{
 				continue;
@@ -2021,9 +2021,9 @@ void FDatasmithC4DImporter::ImportAnimations(TSharedPtr<IDatasmithActorElement>
 				continue;
 			}
 
-			for (melange::CTrack* Track = Tag->GetFirstCTrack(); Track; Track=Track->GetNext())
+			for (cineware::CTrack* Track = Tag->GetFirstCTrack(); Track; Track=Track->GetNext())
 			{
-				melange::DescID TrackDescID = Track->GetDescriptionID();
+				cineware::DescID TrackDescID = Track->GetDescriptionID();
 
 				int32 Depth = TrackDescID.GetDepth();
 				if (Depth != 1)
@@ -2031,31 +2031,31 @@ void FDatasmithC4DImporter::ImportAnimations(TSharedPtr<IDatasmithActorElement>
 					continue;
 				}
 
-				melange::Int32 MelangeTransformType = TrackDescID[0].id;
-				if (MelangeTransformType != melange::ALIGNTOSPLINETAG_POSITION)
+				cineware::Int32 MelangeTransformType = TrackDescID[0].id;
+				if (MelangeTransformType != cineware::ALIGNTOSPLINETAG_POSITION)
 				{
 					continue;
 				}
 
-				melange::CCurve* Curve = Track->GetCurve();
+				cineware::CCurve* Curve = Track->GetCurve();
 				if (!Curve)
 				{
 					continue;
 				}
 
 				// We need to bake every keyframe, as we need to eval the richcurves for the spline position
-				melange::BaseTime StartTime = Curve->GetStartTime();
-				melange::BaseTime EndTime = Curve->GetEndTime();
+				cineware::BaseTime StartTime = Curve->GetStartTime();
+				cineware::BaseTime EndTime = Curve->GetEndTime();
 				int32 FirstFrame = StartTime.GetFrame(MelangeFPS);
 				int32 LastFrame = EndTime.GetFrame(MelangeFPS);
 				for (int32 FrameNumber = FirstFrame; FrameNumber <= LastFrame; ++FrameNumber)
 				{
 					// Uses the timing curve to find the percentage of the spline path at which we must sample
 					// (e.g. 0.0 -> start; 0.5 -> middle; 1.0 -> end)
-					float Percent = (float)Curve->GetValue(melange::BaseTime(StartTime.Get() + melange::Float((FrameNumber-FirstFrame) * (1.0 / MelangeFPS))));
+					float Percent = (float)Curve->GetValue(cineware::BaseTime(StartTime.Get() + cineware::Float((FrameNumber-FirstFrame) * (1.0 / MelangeFPS))));
 
 					// Target spline point in our local space
-					melange::Vector Location = WorldToLocal * melange::Vector((*FoundSpline)[0].Eval(Percent),
+					cineware::Vector Location = WorldToLocal * cineware::Vector((*FoundSpline)[0].Eval(Percent),
 																			  (*FoundSpline)[1].Eval(Percent),
 																			  (*FoundSpline)[2].Eval(Percent));
 					for (int32 Component = 0; Component < 3; ++Component)
@@ -2067,7 +2067,7 @@ void FDatasmithC4DImporter::ImportAnimations(TSharedPtr<IDatasmithActorElement>
 							Component,
 							EDatasmithTransformType::Translation,
 							ComponentValue,
-							melange::ID_BASEOBJECT_REL_POSITION,
+							cineware::ID_BASEOBJECT_REL_POSITION,
 							InitialSize,
 							TransformFrames,
 							InitialValues);
@@ -2078,30 +2078,30 @@ void FDatasmithC4DImporter::ImportAnimations(TSharedPtr<IDatasmithActorElement>
 	}
 
 	// Get the last point in time where we have a valid key
-	melange::BaseTime MaxTime(-1.0);
-	for(melange::CTrack* Track = Object->GetFirstCTrack(); Track; Track=Track->GetNext())
+	cineware::BaseTime MaxTime(-1.0);
+	for(cineware::CTrack* Track = Object->GetFirstCTrack(); Track; Track=Track->GetNext())
 	{
-		melange::DescID TrackDescID = Track->GetDescriptionID();
+		cineware::DescID TrackDescID = Track->GetDescriptionID();
 		if (TrackDescID.GetDepth() != 2)
 		{
 			continue;
 		}
 
-		if (TrackDescID[1].id != melange::VECTOR_X && TrackDescID[1].id != melange::VECTOR_Y && TrackDescID[1].id != melange::VECTOR_Z)
+		if (TrackDescID[1].id != cineware::VECTOR_X && TrackDescID[1].id != cineware::VECTOR_Y && TrackDescID[1].id != cineware::VECTOR_Z)
 		{
 			continue;
 		}
 
-		if (melange::CCurve* Curve = Track->GetCurve())
+		if (cineware::CCurve* Curve = Track->GetCurve())
 		{
 			MaxTime = FMath::Max(MaxTime, Curve->GetEndTime());
 		}
 	}
 
 	// Import animations on the object's attributes
-	for(melange::CTrack* Track = Object->GetFirstCTrack(); Track; Track=Track->GetNext())
+	for(cineware::CTrack* Track = Object->GetFirstCTrack(); Track; Track=Track->GetNext())
 	{
-		melange::DescID TrackDescID = Track->GetDescriptionID();
+		cineware::DescID TrackDescID = Track->GetDescriptionID();
 		if (TrackDescID.GetDepth() != 2)
 		{
 			continue;
@@ -2110,36 +2110,36 @@ void FDatasmithC4DImporter::ImportAnimations(TSharedPtr<IDatasmithActorElement>
 		int32 TransformVectorIndex;
 		switch (TrackDescID[1].id)
 		{
-		case melange::VECTOR_X: TransformVectorIndex = 0; break;
-		case melange::VECTOR_Y: TransformVectorIndex = 1; break;
-		case melange::VECTOR_Z: TransformVectorIndex = 2; break;
+		case cineware::VECTOR_X: TransformVectorIndex = 0; break;
+		case cineware::VECTOR_Y: TransformVectorIndex = 1; break;
+		case cineware::VECTOR_Z: TransformVectorIndex = 2; break;
 		default: continue;
 		}
 
 		EDatasmithTransformType TransformType;
-		melange::Int32 MelangeTransformType = TrackDescID[0].id;
+		cineware::Int32 MelangeTransformType = TrackDescID[0].id;
 		switch (MelangeTransformType)
 		{
-		case melange::ID_BASEOBJECT_REL_POSITION: TransformType = EDatasmithTransformType::Translation; break;
-		case melange::ID_BASEOBJECT_REL_ROTATION: TransformType = EDatasmithTransformType::Rotation; break;
+		case cineware::ID_BASEOBJECT_REL_POSITION: TransformType = EDatasmithTransformType::Translation; break;
+		case cineware::ID_BASEOBJECT_REL_ROTATION: TransformType = EDatasmithTransformType::Rotation; break;
 		case 1100: /*size*/
-		case melange::ID_BASEOBJECT_REL_SCALE:	  TransformType = EDatasmithTransformType::Scale; break;
+		case cineware::ID_BASEOBJECT_REL_SCALE:	  TransformType = EDatasmithTransformType::Scale; break;
 		default: continue;
 		}
 
-		melange::CCurve* Curve = Track->GetCurve();
+		cineware::CCurve* Curve = Track->GetCurve();
 		if (!Curve)
 		{
 			continue;
 		}
 
-		for (melange::Int32 KeyIndex = 0; KeyIndex < Curve->GetKeyCount(); KeyIndex++)
+		for (cineware::Int32 KeyIndex = 0; KeyIndex < Curve->GetKeyCount(); KeyIndex++)
 		{
-			melange::CKey* CurrentKey = Curve->GetKey(KeyIndex);
-			melange::CINTERPOLATION Interpolation = CurrentKey->GetInterpolation();
+			cineware::CKey* CurrentKey = Curve->GetKey(KeyIndex);
+			cineware::CINTERPOLATION Interpolation = CurrentKey->GetInterpolation();
 
 			int32 FrameNumber = CurrentKey->GetTime().GetFrame(MelangeFPS);
-			melange::Float FrameValue = CurrentKey->GetValue();
+			cineware::Float FrameValue = CurrentKey->GetValue();
 			AddFrameValueToAnimMap(
 				Object,
 				FrameNumber,
@@ -2154,16 +2154,16 @@ void FDatasmithC4DImporter::ImportAnimations(TSharedPtr<IDatasmithActorElement>
 			if (KeyIndex < Curve->GetKeyCount() - 1)
 			{
 				//"Bake" the animation by generating a key for each frame between this Key and the next one
-				melange::CKey* NextKey = Curve->GetKey(KeyIndex+1);
-				melange::BaseTime CurrentKeyTime = CurrentKey->GetTime();
-				melange::BaseTime NextKeyTime = NextKey->GetTime();
+				cineware::CKey* NextKey = Curve->GetKey(KeyIndex+1);
+				cineware::BaseTime CurrentKeyTime = CurrentKey->GetTime();
+				cineware::BaseTime NextKeyTime = NextKey->GetTime();
 				int32 NextKeyFrameNumber = NextKeyTime.GetFrame(MelangeFPS);
 				int32 FrameCount = NextKeyFrameNumber - FrameNumber;
-				melange::Float ElapsedTime = NextKeyTime.Get() - CurrentKeyTime.Get();
+				cineware::Float ElapsedTime = NextKeyTime.Get() - CurrentKeyTime.Get();
 				for (int32 FrameIndex = 1; FrameIndex < FrameCount; FrameIndex++)
 				{
 					FrameNumber++;
-					FrameValue = Curve->GetValue(melange::BaseTime(CurrentKeyTime.Get() + ((ElapsedTime / FrameCount) * FrameIndex)));
+					FrameValue = Curve->GetValue(cineware::BaseTime(CurrentKeyTime.Get() + ((ElapsedTime / FrameCount) * FrameIndex)));
 					AddFrameValueToAnimMap(
 						Object,
 						FrameNumber,
@@ -2181,10 +2181,10 @@ void FDatasmithC4DImporter::ImportAnimations(TSharedPtr<IDatasmithActorElement>
 		// Make sure the transform frame values remain at their last valid value up until the end of the
 		// animation. We use FVectors to store all three components at once, if we don't do this we will
 		// incorrectly think that components whose animation curves end early have gone back to zero
-		melange::Float LastValue = Curve->GetValue(Curve->GetEndTime());
-		melange::Int32 FirstFrameToFill = Curve->GetEndTime().GetFrame(MelangeFPS) + 1;
-		melange::Int32 LastFrameToFill = MaxTime.GetFrame(MelangeFPS);
-		for (melange::Int32 Frame = FirstFrameToFill; Frame <= LastFrameToFill; ++Frame)
+		cineware::Float LastValue = Curve->GetValue(Curve->GetEndTime());
+		cineware::Int32 FirstFrameToFill = Curve->GetEndTime().GetFrame(MelangeFPS) + 1;
+		cineware::Int32 LastFrameToFill = MaxTime.GetFrame(MelangeFPS);
+		for (cineware::Int32 Frame = FirstFrameToFill; Frame <= LastFrameToFill; ++Frame)
 		{
 			AddFrameValueToAnimMap(
 				Object,
@@ -2234,10 +2234,10 @@ void FDatasmithC4DImporter::ImportAnimations(TSharedPtr<IDatasmithActorElement>
 		VisibilityAnimation->AddFrame(FDatasmithVisibilityFrameInfo(LastFrameAdded, true));
 
 		// After our last frame we should be visible, but don't add a new key if that is also the last frame of the document
-		melange::GeData Data;
-		if (C4dDocument->GetParameter(melange::DOCUMENT_MAXTIME, Data) && Data.GetType() == melange::DA_TIME)
+		cineware::GeData Data;
+		if (C4dDocument->GetParameter(cineware::DOCUMENT_MAXTIME, Data) && Data.GetType() == cineware::DA_TIME)
 		{
-			const melange::BaseTime& Time = Data.GetTime();
+			const cineware::BaseTime& Time = Data.GetTime();
 			int32 LastDocumentFrame = Time.GetFrame(MelangeFPS);
 
 			if (LastFrameAdded < LastDocumentFrame)
@@ -2286,9 +2286,9 @@ void FDatasmithC4DImporter::ImportAnimations(TSharedPtr<IDatasmithActorElement>
 				// Copy as we might be reusing a LastValue
 				FVector3f TransformValueCopy = (FVector3f)*TransformValue;
 
-				// If the object is in the HPB rotation order, melange will store its euler rotation
+				// If the object is in the HPB rotation order, cineware will store its euler rotation
 				// as "H, P, B", basically storing the rotations as "YXZ". Lets switch it back to XYZ
-				if (RotationOrder == melange::ROTATIONORDER_HPB)
+				if (RotationOrder == cineware::ROTATIONORDER_HPB)
 				{
 					Swap(TransformValueCopy.X, TransformValueCopy.Y);
 				}
@@ -2308,23 +2308,23 @@ void FDatasmithC4DImporter::ImportAnimations(TSharedPtr<IDatasmithActorElement>
 				FQuat FinalQuat;
 				switch (RotationOrder)
 				{
-				case melange::ROTATIONORDER_XZYGLOBAL:
+				case cineware::ROTATIONORDER_XZYGLOBAL:
 					FinalQuat = YRot * ZRot * XRot;
 					break;
-				case melange::ROTATIONORDER_XYZGLOBAL:
+				case cineware::ROTATIONORDER_XYZGLOBAL:
 					FinalQuat = ZRot * YRot * XRot;
 					break;
-				case melange::ROTATIONORDER_YZXGLOBAL:
+				case cineware::ROTATIONORDER_YZXGLOBAL:
 					FinalQuat = XRot * ZRot * YRot;
 					break;
-				case melange::ROTATIONORDER_ZYXGLOBAL:
+				case cineware::ROTATIONORDER_ZYXGLOBAL:
 					FinalQuat = XRot * YRot * ZRot;
 					break;
-				case melange::ROTATIONORDER_YXZGLOBAL:
+				case cineware::ROTATIONORDER_YXZGLOBAL:
 					FinalQuat = ZRot * XRot * YRot;
 					break;
-				case melange::ROTATIONORDER_ZXYGLOBAL:
-				case melange::ROTATIONORDER_HPB:
+				case cineware::ROTATIONORDER_ZXYGLOBAL:
+				case cineware::ROTATIONORDER_HPB:
 				default:
 					FinalQuat = YRot * XRot * ZRot;
 					break;
@@ -2360,16 +2360,16 @@ FVector FDatasmithC4DImporter::GetDocumentDefaultColor()
 {
 	if (!DefaultDocumentColorLinear.IsSet())
 	{
-		const int32 DefaultColorType = MelangeGetInt32(C4dDocument, melange::DOCUMENT_DEFAULTMATERIAL_TYPE);
+		const int32 DefaultColorType = MelangeGetInt32(C4dDocument, cineware::DOCUMENT_DEFAULTMATERIAL_TYPE);
 		switch (DefaultColorType)
 		{
-		case melange::DOCUMENT_DEFAULTMATERIAL_TYPE_WHITE: // This says "80% Gray" on the UI
+		case cineware::DOCUMENT_DEFAULTMATERIAL_TYPE_WHITE: // This says "80% Gray" on the UI
 			DefaultDocumentColorLinear = FVector(0.603828f, 0.603828f, 0.603828f);
 			break;
-		case melange::DOCUMENT_DEFAULTMATERIAL_TYPE_USER:
-			DefaultDocumentColorLinear = MelangeGetColor(C4dDocument, melange::DOCUMENT_DEFAULTMATERIAL_COLOR);
+		case cineware::DOCUMENT_DEFAULTMATERIAL_TYPE_USER:
+			DefaultDocumentColorLinear = MelangeGetColor(C4dDocument, cineware::DOCUMENT_DEFAULTMATERIAL_COLOR);
 			break;
-		case melange::DOCUMENT_DEFAULTMATERIAL_TYPE_BLUE: // Intended fall-through. Blue is the default
+		case cineware::DOCUMENT_DEFAULTMATERIAL_TYPE_BLUE: // Intended fall-through. Blue is the default
 		default:
 			DefaultDocumentColorLinear = FVector(0.099899f, 0.116971f, 0.138432f);
 		}
@@ -2378,28 +2378,28 @@ FVector FDatasmithC4DImporter::GetDocumentDefaultColor()
 	return DefaultDocumentColorLinear.GetValue();
 }
 
-TArray<melange::TextureTag*> FDatasmithC4DImporter::GetActiveTextureTags(const melange::BaseObject* Object, const TArray<melange::TextureTag*>& OrderedTextureTags)
+TArray<cineware::TextureTag*> FDatasmithC4DImporter::GetActiveTextureTags(const cineware::BaseObject* Object, const TArray<cineware::TextureTag*>& OrderedTextureTags)
 {
 	if (!Object)
 	{
 		return {};
 	}
 
-	TArray<melange::BaseSelect*> OrderedSelectionTags;
+	TArray<cineware::BaseSelect*> OrderedSelectionTags;
 	OrderedSelectionTags.Add(nullptr); // "unselected" group
 
-	TMap<FString, melange::BaseSelect*> SelectionTagsByName;
+	TMap<FString, cineware::BaseSelect*> SelectionTagsByName;
 
 	// Fetch selection tags, which only affect this polygon
 	// The texture tags are fetched when moving down the hierarchy, as texture tags on parents also affect children
-	for (melange::BaseTag* Tag = const_cast<melange::BaseObject*>(Object)->GetFirstTag(); Tag; Tag = Tag->GetNext())
+	for (cineware::BaseTag* Tag = const_cast<cineware::BaseObject*>(Object)->GetFirstTag(); Tag; Tag = Tag->GetNext())
 	{
 		if (Tag->GetType() == Tpolygonselection)
 		{
-			FString SelectionName = MelangeGetString(Tag, melange::POLYGONSELECTIONTAG_NAME);
+			FString SelectionName = MelangeGetString(Tag, cineware::POLYGONSELECTIONTAG_NAME);
 			if (!SelectionName.IsEmpty())
 			{
-				melange::BaseSelect* BaseSelect = static_cast<melange::SelectionTag*>(Tag)->GetBaseSelect();
+				cineware::BaseSelect* BaseSelect = static_cast<cineware::SelectionTag*>(Tag)->GetBaseSelect();
 				OrderedSelectionTags.Add(BaseSelect);
 
 				SelectionTagsByName.Add(SelectionName, BaseSelect);
@@ -2409,21 +2409,21 @@ TArray<melange::TextureTag*> FDatasmithC4DImporter::GetActiveTextureTags(const m
 
 	// If we have multiple texture tags using the same selection, only the latter one will be applied
 	// Order is important here: We must scan TextureTags front to back to guarantee that behavior
-	TMap<melange::BaseSelect*, melange::TextureTag*> ActiveTextureTags;
-	for (melange::TextureTag* TextureTag : OrderedTextureTags)
+	TMap<cineware::BaseSelect*, cineware::TextureTag*> ActiveTextureTags;
+	for (cineware::TextureTag* TextureTag : OrderedTextureTags)
 	{
-		melange::BaseSelect* UsedSelectionTag = nullptr;
+		cineware::BaseSelect* UsedSelectionTag = nullptr;
 
-		FString UsedSelectionTagName = MelangeGetString(TextureTag, melange::TEXTURETAG_RESTRICTION);
+		FString UsedSelectionTagName = MelangeGetString(TextureTag, cineware::TEXTURETAG_RESTRICTION);
 		if (!UsedSelectionTagName.IsEmpty())
 		{
-			if (melange::BaseSelect** FoundSelectionTag = SelectionTagsByName.Find(UsedSelectionTagName))
+			if (cineware::BaseSelect** FoundSelectionTag = SelectionTagsByName.Find(UsedSelectionTagName))
 			{
 				UsedSelectionTag = *FoundSelectionTag;
 			}
 		}
 
-		melange::BaseList2D* TextureMaterial = TextureTag ? MelangeGetLink(TextureTag, melange::TEXTURETAG_MATERIAL) : nullptr;
+		cineware::BaseList2D* TextureMaterial = TextureTag ? MelangeGetLink(TextureTag, cineware::TEXTURETAG_MATERIAL) : nullptr;
 
 		// Note: If this texture tag is applied without a polygon selection, UsedSelectionTag will be
 		// nullptr here, but that is intentional: It's how we signal the "unselected" selection group
@@ -2432,17 +2432,17 @@ TArray<melange::TextureTag*> FDatasmithC4DImporter::GetActiveTextureTags(const m
 
 	// Order is important: The polygon groups are created according to the order with which we retrieve our selection tags,
 	// so the order with which we return these texture tags must match it exactly
-	TArray<melange::TextureTag*> OrderedActiveTextureTags;
-	for (melange::BaseSelect* SelectionTag : OrderedSelectionTags)
+	TArray<cineware::TextureTag*> OrderedActiveTextureTags;
+	for (cineware::BaseSelect* SelectionTag : OrderedSelectionTags)
 	{
-		melange::TextureTag** TextureTagForSelection = ActiveTextureTags.Find(SelectionTag);
+		cineware::TextureTag** TextureTagForSelection = ActiveTextureTags.Find(SelectionTag);
 		OrderedActiveTextureTags.Add(TextureTagForSelection ? *TextureTagForSelection : nullptr);
 	}
 
 	return OrderedActiveTextureTags;
 }
 
-void FDatasmithC4DImporter::RegisterInstancedHierarchy(melange::BaseObject* InstanceSubObject, melange::BaseObject* OriginalSubObject)
+void FDatasmithC4DImporter::RegisterInstancedHierarchy(cineware::BaseObject* InstanceSubObject, cineware::BaseObject* OriginalSubObject)
 {
 	if (!InstanceSubObject || !OriginalSubObject)
 	{
@@ -2457,21 +2457,21 @@ void FDatasmithC4DImporter::RegisterInstancedHierarchy(melange::BaseObject* Inst
 
 void FDatasmithC4DImporter::RedirectInstancedAnimations()
 {
-	for (TPair<IDatasmithActorElement*, melange::BaseObject*>& Entry : ActorElementToAnimationSources)
+	for (TPair<IDatasmithActorElement*, cineware::BaseObject*>& Entry : ActorElementToAnimationSources)
 	{
-		if (melange::BaseObject** OriginalObject = InstancedSubObjectsToOriginals.Find(Entry.Value))
+		if (cineware::BaseObject** OriginalObject = InstancedSubObjectsToOriginals.Find(Entry.Value))
 		{
 			Entry.Value = *OriginalObject;
 		}
 	}
 }
 
-TSharedPtr<IDatasmithActorElement> FDatasmithC4DImporter::ImportObjectAndChildren(melange::BaseObject* ActorObject, melange::BaseObject* DataObject, TSharedPtr<IDatasmithActorElement> ParentActor, const melange::Matrix& WorldTransformMatrix, const FString& InstancePath, const FString& DatasmithLabel, const TArray<melange::TextureTag*>& TextureTags)
+TSharedPtr<IDatasmithActorElement> FDatasmithC4DImporter::ImportObjectAndChildren(cineware::BaseObject* ActorObject, cineware::BaseObject* DataObject, TSharedPtr<IDatasmithActorElement> ParentActor, const cineware::Matrix& WorldTransformMatrix, const FString& InstancePath, const FString& DatasmithLabel, const TArray<cineware::TextureTag*>& TextureTags)
 {
 	TSharedPtr<IDatasmithActorElement> ActorElement;
-	melange::Int32 ObjectType = DataObject->GetType();
-	melange::BaseObject* ActorCache = GetBestMelangeCache(ActorObject);
-	melange::BaseObject* DataCache = GetBestMelangeCache(DataObject);
+	cineware::Int32 ObjectType = DataObject->GetType();
+	cineware::BaseObject* ActorCache = GetBestMelangeCache(ActorObject);
+	cineware::BaseObject* DataCache = GetBestMelangeCache(DataObject);
 	if (!DataCache)
 	{
 		DataCache = ActorCache;
@@ -2493,12 +2493,12 @@ TSharedPtr<IDatasmithActorElement> FDatasmithC4DImporter::ImportObjectAndChildre
 		DatasmithName = MD5FromString(InstancePath) + "_" + DatasmithName.GetValue();
 	}
 
-	melange::Matrix NewWorldTransformMatrix = WorldTransformMatrix * ActorObject->GetMl();
+	cineware::Matrix NewWorldTransformMatrix = WorldTransformMatrix * ActorObject->GetMl();
 
 	// Fetch actor layer
 	FString TargetLayerName;
 	bool bActorVisible = true;
-	if (melange::LayerObject* LayerObject = static_cast<melange::LayerObject*>(MelangeGetLink(ActorObject, melange::ID_LAYER_LINK)))
+	if (cineware::LayerObject* LayerObject = static_cast<cineware::LayerObject*>(MelangeGetLink(ActorObject, cineware::ID_LAYER_LINK)))
 	{
 		// Do not create actors from invisible layers
 		// We may end up creating null actors if the actor is in an invisible layer, and even
@@ -2507,7 +2507,7 @@ TSharedPtr<IDatasmithActorElement> FDatasmithC4DImporter::ImportObjectAndChildre
 		// and so on.
 		// Exceptions are generators: If a cloner is in an invisible layer, the child nodes are always invisible,
 		// and also if the cloner is in a visible layer, the child nodes are always visible
-		bActorVisible = MelangeGetBool(LayerObject, melange::ID_LAYER_VIEW);
+		bActorVisible = MelangeGetBool(LayerObject, cineware::ID_LAYER_VIEW);
 
 		TargetLayerName = MelangeObjectName(LayerObject);
 	}
@@ -2538,12 +2538,12 @@ TSharedPtr<IDatasmithActorElement> FDatasmithC4DImporter::ImportObjectAndChildre
 			//	| | | -CACHE: Cube 0(Opolygon)
 			//	| -Cube(Ocube)
 
-			if (ObjectType == Ocloner && MelangeGetInt32(ActorObject, melange::MGCLONER_VOLUMEINSTANCES_MODE) != 0)
+			if (ObjectType == Ocloner && MelangeGetInt32(ActorObject, cineware::MGCLONER_VOLUMEINSTANCES_MODE) != 0)
 			{
 				// Render/Multi-instance cloner should be ignored
 				UE_LOG(LogDatasmithC4DImport, Warning, TEXT("Render-instance or multi-instance Cloners are not supported. Actor '%s' will be ignored"), *MelangeObjectName(ActorObject));
 			}
-			else if (ObjectType == Oarray && MelangeGetInt32(ActorObject, melange::ARRAYOBJECT_RENDERINSTANCES) != 0)
+			else if (ObjectType == Oarray && MelangeGetInt32(ActorObject, cineware::ARRAYOBJECT_RENDERINSTANCES) != 0)
 			{
 				// Render-instance arrays should be ignored
 				UE_LOG(LogDatasmithC4DImport, Warning, TEXT("Render-instance Arrays are not supported. Actor '%s' will be ignored"), *MelangeObjectName(ActorObject));
@@ -2597,7 +2597,7 @@ TSharedPtr<IDatasmithActorElement> FDatasmithC4DImporter::ImportObjectAndChildre
 			break;
 
 		case Oinstance:
-			if (melange::BaseObject* InstanceLink = static_cast<melange::BaseObject*>(MelangeGetLink(DataObject, melange::INSTANCEOBJECT_LINK)))
+			if (cineware::BaseObject* InstanceLink = static_cast<cineware::BaseObject*>(MelangeGetLink(DataObject, cineware::INSTANCEOBJECT_LINK)))
 			{
 				if (TOptional<FString> ObjectID = MelangeObjectID(DataObject))
 				{
@@ -2632,7 +2632,7 @@ TSharedPtr<IDatasmithActorElement> FDatasmithC4DImporter::ImportObjectAndChildre
 			break;
 
 		case Ospline:
-			if (melange::SplineObject* Spline = static_cast<melange::SplineObject*>(ActorObject))
+			if (cineware::SplineObject* Spline = static_cast<cineware::SplineObject*>(ActorObject))
 			{
 				ImportSpline(Spline);
 			}
@@ -2649,10 +2649,10 @@ TSharedPtr<IDatasmithActorElement> FDatasmithC4DImporter::ImportObjectAndChildre
 		}
 		else if (ObjectType == Opolygon)
 		{
-			melange::PolygonObject* PolygonObject = static_cast<melange::PolygonObject*>(DataObject);
+			cineware::PolygonObject* PolygonObject = static_cast<cineware::PolygonObject*>(DataObject);
 			if (Options.bImportEmptyMesh || PolygonObject->GetPolygonCount() > 0)
 			{
-				TArray<melange::TextureTag*> ActiveTextureTags = GetActiveTextureTags(PolygonObject, TextureTags);
+				TArray<cineware::TextureTag*> ActiveTextureTags = GetActiveTextureTags(PolygonObject, TextureTags);
 				if (TSharedPtr<IDatasmithMeshActorElement> MeshActorElement = ImportPolygon(PolygonObject, DatasmithName.GetValue(), DatasmithLabel, ActiveTextureTags))
 				{
 					ActorElement = MeshActorElement;
@@ -2712,7 +2712,7 @@ TSharedPtr<IDatasmithActorElement> FDatasmithC4DImporter::ImportObjectAndChildre
 	return ActorElement;
 }
 
-void FDatasmithC4DImporter::ImportHierarchy(melange::BaseObject* ActorObject, melange::BaseObject* DataObject, TSharedPtr<IDatasmithActorElement> ParentActor, const melange::Matrix& WorldTransformMatrix, const FString& InstancePath, const TArray<melange::TextureTag*>& TextureTags)
+void FDatasmithC4DImporter::ImportHierarchy(cineware::BaseObject* ActorObject, cineware::BaseObject* DataObject, TSharedPtr<IDatasmithActorElement> ParentActor, const cineware::Matrix& WorldTransformMatrix, const FString& InstancePath, const TArray<cineware::TextureTag*>& TextureTags)
 {
 	while (ActorObject || DataObject)
 	{
@@ -2726,12 +2726,12 @@ void FDatasmithC4DImporter::ImportHierarchy(melange::BaseObject* ActorObject, me
 		}
 
 		// Reset this for every child as texture tags only propagate down, not between siblings
-		TArray<melange::TextureTag*> TextureTagsDown = TextureTags;
+		TArray<cineware::TextureTag*> TextureTagsDown = TextureTags;
 
 		bool SkipObject = false;
-		for (melange::BaseTag* Tag = ActorObject->GetFirstTag(); Tag; Tag = Tag->GetNext())
+		for (cineware::BaseTag* Tag = ActorObject->GetFirstTag(); Tag; Tag = Tag->GetNext())
 		{
-			melange::Int32 TagType = Tag->GetType();
+			cineware::Int32 TagType = Tag->GetType();
 			if (TagType == Tannotation)
 			{
 				FString AnnotationLabel = MelangeGetString(Tag, 10014);
@@ -2743,7 +2743,7 @@ void FDatasmithC4DImporter::ImportHierarchy(melange::BaseObject* ActorObject, me
 			}
 			else if (TagType == Ttexture)
 			{
-				TextureTagsDown.Add(static_cast<melange::TextureTag*>(Tag));
+				TextureTagsDown.Add(static_cast<cineware::TextureTag*>(Tag));
 			}
 		}
 
@@ -2758,32 +2758,32 @@ void FDatasmithC4DImporter::ImportHierarchy(melange::BaseObject* ActorObject, me
 	}
 }
 
-TSharedPtr<IDatasmithMeshElement> FDatasmithC4DImporter::ImportMesh(melange::PolygonObject* PolyObject, const FString& DatasmithMeshName, const FString& DatasmithLabel)
+TSharedPtr<IDatasmithMeshElement> FDatasmithC4DImporter::ImportMesh(cineware::PolygonObject* PolyObject, const FString& DatasmithMeshName, const FString& DatasmithLabel)
 {
-	melange::Int32 PointCount = PolyObject->GetPointCount();
-	melange::Int32 PolygonCount = PolyObject->GetPolygonCount();
+	cineware::Int32 PointCount = PolyObject->GetPointCount();
+	cineware::Int32 PolygonCount = PolyObject->GetPolygonCount();
 
-	const melange::Vector* Points = PolyObject->GetPointR();
-	const melange::CPolygon* Polygons = PolyObject->GetPolygonR();
+	const cineware::Vector* Points = PolyObject->GetPointR();
+	const cineware::CPolygon* Polygons = PolyObject->GetPolygonR();
 
 	// Get vertex normals
-	melange::Vector32* Normals = nullptr;
+	cineware::Vector32* Normals = nullptr;
 	if (PolyObject->GetTag(Tphong))
 	{
 		Normals = PolyObject->CreatePhongNormals();
 	}
 
 	// Collect all UV channels and material slot information for this PolygonObject
-	TArray<melange::ConstUVWHandle> UVWTagsData;
-	TArray<melange::BaseSelect*> SelectionTags;
+	TArray<cineware::ConstUVWHandle> UVWTagsData;
+	TArray<cineware::BaseSelect*> SelectionTags;
 	SelectionTags.Add(nullptr); // The "unselected" group
 
-	for (melange::BaseTag* Tag = PolyObject->GetFirstTag(); Tag; Tag = Tag->GetNext())
+	for (cineware::BaseTag* Tag = PolyObject->GetFirstTag(); Tag; Tag = Tag->GetNext())
 	{
-		melange::Int32 TagType = Tag->GetType();
+		cineware::Int32 TagType = Tag->GetType();
 		if (TagType == Tuvw)
 		{
-			melange::UVWTag* UVWTag = static_cast<melange::UVWTag*>(Tag);
+			cineware::UVWTag* UVWTag = static_cast<cineware::UVWTag*>(Tag);
 			if (UVWTag->GetDataCount() == PolygonCount)
 			{
 				UVWTagsData.Add(UVWTag->GetDataAddressR());
@@ -2795,10 +2795,10 @@ TSharedPtr<IDatasmithMeshElement> FDatasmithC4DImporter::ImportMesh(melange::Pol
 		}
 		else if (TagType == Tpolygonselection)
 		{
-			FString SelectionName = MelangeGetString(Tag, melange::POLYGONSELECTIONTAG_NAME);
+			FString SelectionName = MelangeGetString(Tag, cineware::POLYGONSELECTIONTAG_NAME);
 			if (!SelectionName.IsEmpty())
 			{
-				SelectionTags.Add(static_cast<melange::SelectionTag*>(Tag)->GetBaseSelect());
+				SelectionTags.Add(static_cast<cineware::SelectionTag*>(Tag)->GetBaseSelect());
 			}
 		}
 	}
@@ -2872,7 +2872,7 @@ TSharedPtr<IDatasmithMeshElement> FDatasmithC4DImporter::ImportMesh(melange::Pol
 	// Create polygons
 	for (int32 PolygonIndex = 0; PolygonIndex < PolygonCount; ++PolygonIndex)
 	{
-		const melange::CPolygon& Polygon = Polygons[PolygonIndex];
+		const cineware::CPolygon& Polygon = Polygons[PolygonIndex];
 
 		// Check if we're a triangle or a quad
 		if (Polygon.c == Polygon.d)
@@ -2905,7 +2905,7 @@ TSharedPtr<IDatasmithMeshElement> FDatasmithC4DImporter::ImportMesh(melange::Pol
 
 			// Check if those vertices lead to degenerate triangles first, to prevent us from ever adding unused data to the MeshDescription
 			FVector RawNormal = ( ( TriangleVertexPositions[ 1 ] - TriangleVertexPositions[ 2 ] ) ^ ( TriangleVertexPositions[ 0 ] - TriangleVertexPositions[ 2 ] ) );
-			if ( RawNormal.SizeSquared() < SMALL_NUMBER )
+			if ( RawNormal.SizeSquared() < UE_SMALL_NUMBER )
 			{
 				continue;
 			}
@@ -2917,7 +2917,7 @@ TSharedPtr<IDatasmithMeshElement> FDatasmithC4DImporter::ImportMesh(melange::Pol
 			}
 		}
 
-		// Fetch melange polygon normals (always 4, even if triangle)
+		// Fetch cineware polygon normals (always 4, even if triangle)
 		if (Normals)
 		{
 			for (int32 VertexIndex = 0; VertexIndex < 4; ++VertexIndex)
@@ -2937,15 +2937,15 @@ TSharedPtr<IDatasmithMeshElement> FDatasmithC4DImporter::ImportMesh(melange::Pol
 		// UVs
 		for (int32 ChannelIndex = 0; ChannelIndex < UVChannelCount; ++ChannelIndex)
 		{
-			melange::ConstUVWHandle UVWTagData = UVWTagsData[ChannelIndex];
-			melange::UVWStruct UVWStruct;
-			melange::UVWTag::Get(UVWTagData, PolygonIndex, UVWStruct);
-			melange::Vector* UVs = &UVWStruct.a;
+			cineware::ConstUVWHandle UVWTagData = UVWTagsData[ChannelIndex];
+			cineware::UVWStruct UVWStruct;
+			cineware::UVWTag::Get(UVWTagData, PolygonIndex, UVWStruct);
+			cineware::Vector* UVs = &UVWStruct.a;
 
-			// Fetch melange UVs
+			// Fetch cineware UVs
 			for (int32 VertexIndex = 0; VertexIndex < 4; ++VertexIndex)
 			{
-				melange::Vector& PointUVs = UVs[VertexIndex];
+				cineware::Vector& PointUVs = UVs[VertexIndex];
 				FVector2f& UnrealUVs = QuadUVs[VertexIndex];
 
 				if (PointUVs.z != 0.0f && PointUVs.z != 1.0f)
@@ -2980,7 +2980,7 @@ TSharedPtr<IDatasmithMeshElement> FDatasmithC4DImporter::ImportMesh(melange::Pol
 		int32 SlotIndex = NumSlots - 1;
 		for (; SlotIndex > 0; --SlotIndex)
 		{
-			melange::BaseSelect* SelectionTag = SelectionTags[SlotIndex];
+			cineware::BaseSelect* SelectionTag = SelectionTags[SlotIndex];
 
 			if (SelectionTag && SelectionTag->IsSelected(PolygonIndex))
 			{
@@ -3010,7 +3010,7 @@ TSharedPtr<IDatasmithMeshElement> FDatasmithC4DImporter::ImportMesh(melange::Pol
 
 	if (Normals)
 	{
-		melange::DeleteMem(Normals);
+		cineware::DeleteMem(Normals);
 	}
 
 	TSharedRef<IDatasmithMeshElement> MeshElement = FDatasmithSceneFactory::CreateMesh(*DatasmithMeshName);
@@ -3039,7 +3039,7 @@ bool FDatasmithC4DImporter::OpenFile(const FString& InFilename)
 {
 	SCOPE_CYCLE_COUNTER(STAT_C4DImporter_LoadFile)
 
-	using namespace melange;
+	using namespace cineware;
 
 	if (!FPaths::FileExists(InFilename))
 	{
@@ -3104,9 +3104,9 @@ bool FDatasmithC4DImporter::OpenFile(const FString& InFilename)
 	return true;
 }
 
-melange::BaseObject* FDatasmithC4DImporter::FindMelangeObject(const FString& SearchObjectID, melange::BaseObject* Object)
+cineware::BaseObject* FDatasmithC4DImporter::FindMelangeObject(const FString& SearchObjectID, cineware::BaseObject* Object)
 {
-	melange::BaseObject* FoundObject = nullptr;
+	cineware::BaseObject* FoundObject = nullptr;
 
 	while (Object && !FoundObject)
 	{
@@ -3131,7 +3131,7 @@ melange::BaseObject* FDatasmithC4DImporter::FindMelangeObject(const FString& Sea
 	return FoundObject;
 }
 
-melange::BaseObject* FDatasmithC4DImporter::GoToMelangeHierarchyPosition(melange::BaseObject* Object, const FString& HierarchyPosition)
+cineware::BaseObject* FDatasmithC4DImporter::GoToMelangeHierarchyPosition(cineware::BaseObject* Object, const FString& HierarchyPosition)
 {
 	if (Object)
 	{
@@ -3163,14 +3163,14 @@ melange::BaseObject* FDatasmithC4DImporter::GoToMelangeHierarchyPosition(melange
 bool FDatasmithC4DImporter::ProcessScene()
 {
 	// Cinema 4D Document settings
-	MelangeFPS = static_cast<melange::Float>(MelangeGetInt32(C4dDocument, melange::DOCUMENT_FPS));
+	MelangeFPS = static_cast<cineware::Float>(MelangeGetInt32(C4dDocument, cineware::DOCUMENT_FPS));
 	if(MelangeFPS == 0.0f)
 	{
 		UE_LOG(LogDatasmithC4DImport, Error, TEXT("DOCUMENT_FPS not found"));
 		return false;
 	}
-	MelangeColorProfile = MelangeGetInt32(C4dDocument, melange::DOCUMENT_COLORPROFILE);
-	melange::RenderData* RenderData = C4dDocument->GetActiveRenderData();
+	MelangeColorProfile = MelangeGetInt32(C4dDocument, cineware::DOCUMENT_COLORPROFILE);
+	cineware::RenderData* RenderData = C4dDocument->GetActiveRenderData();
 	if (!RenderData)
 	{
 		UE_LOG(LogDatasmithC4DImport, Error, TEXT("Active Render Data not found"));
@@ -3189,8 +3189,8 @@ bool FDatasmithC4DImporter::ProcessScene()
 	// Need a RootActor for RemoveEmptyActors and to make AddChildActor agnostic to actor hierarchy level
 	TSharedPtr<IDatasmithActorElement> RootActor = FDatasmithSceneFactory::CreateActor(TEXT("RootActor"));
 	DatasmithScene->AddActor(RootActor);
-	TArray<melange::TextureTag*> TextureTags;
-	ImportHierarchy(C4dDocument->GetFirstObject(), C4dDocument->GetFirstObject(), RootActor, melange::Matrix(), "", TextureTags);
+	TArray<cineware::TextureTag*> TextureTags;
+	ImportHierarchy(C4dDocument->GetFirstObject(), C4dDocument->GetFirstObject(), RootActor, cineware::Matrix(), "", TextureTags);
 
 	// Animations
 	LevelSequence = FDatasmithSceneFactory::CreateLevelSequence(DatasmithScene->GetName());
@@ -3225,7 +3225,7 @@ void FDatasmithC4DImporter::UnloadScene()
 }
 
 // Traverse the LayerObject hierarchy adding visible layer names to VisibleLayers
-void RecursivelyParseLayers(melange::LayerObject* CurrentLayer, TSet<FName>& VisibleLayers)
+void RecursivelyParseLayers(cineware::LayerObject* CurrentLayer, TSet<FName>& VisibleLayers)
 {
 	if (CurrentLayer == nullptr)
 	{
@@ -3234,8 +3234,8 @@ void RecursivelyParseLayers(melange::LayerObject* CurrentLayer, TSet<FName>& Vis
 
 	FString Name = MelangeObjectName(CurrentLayer);
 
-	melange::GeData Data;
-	if (MelangeGetBool(CurrentLayer, melange::ID_LAYER_VIEW))
+	cineware::GeData Data;
+	if (MelangeGetBool(CurrentLayer, cineware::ID_LAYER_VIEW))
 	{
 		VisibleLayers.Add(FName(*Name));
 	}
diff --git a/Engine/Plugins/Enterprise/DatasmithC4DImporter/Source/DatasmithC4DTranslator/Private/DatasmithC4DImporter.h b/Engine/Plugins/Enterprise/DatasmithC4DImporter/Source/DatasmithC4DTranslator/Private/DatasmithC4DImporter.h
--- a/Engine/Plugins/Enterprise/DatasmithC4DImporter/Source/DatasmithC4DTranslator/Private/DatasmithC4DImporter.h
+++ b/Engine/Plugins/Enterprise/DatasmithC4DImporter/Source/DatasmithC4DTranslator/Private/DatasmithC4DImporter.h
@@ -4,6 +4,11 @@
 
 #ifdef _MELANGE_SDK_
 
+#ifdef UE_DEFINE_LEGACY_MATH_CONSTANT_MACRO_NAMES
+#undef UE_DEFINE_LEGACY_MATH_CONSTANT_MACRO_NAMES
+#endif
+#define UE_DEFINE_LEGACY_MATH_CONSTANT_MACRO_NAMES 0
+
 #include "CoreMinimal.h"
 #include "Logging/LogMacros.h"
 #include "MeshDescription.h"
@@ -59,92 +64,92 @@ public:
 	/** Parse the scene contained in the previously opened file and process its content according to parameters from incoming context */
 	bool ProcessScene();
 
-	/** Unload melange resources after importing */
+	/** Unload cineware resources after importing */
 	void UnloadScene();
 
 	void GetGeometriesForMeshElementAndRelease(const TSharedRef<IDatasmithMeshElement> MeshElement, TArray<FMeshDescription>& OutMeshDescriptions);
 	TSharedPtr<IDatasmithLevelSequenceElement> GetLevelSequence();
 
-	/** Finds the most derived cache for a melange object. That will be e.g. a polygon cache or a deformed polygon cache, if it has one */
-	melange::BaseObject* GetBestMelangeCache(melange::BaseObject* Object);
+	/** Finds the most derived cache for a cineware object. That will be e.g. a polygon cache or a deformed polygon cache, if it has one */
+	cineware::BaseObject* GetBestMelangeCache(cineware::BaseObject* Object);
 
-	/** Generates a unique identifier string for a melange object based on the object's name and its position in the hierarchy */
-	TOptional<FString> MelangeObjectID(melange::BaseObject* Object);
+	/** Generates a unique identifier string for a cineware object based on the object's name and its position in the hierarchy */
+	TOptional<FString> MelangeObjectID(cineware::BaseObject* Object);
 
-	/** Searches the melange object hierarchy for a melange::BaseObject that has a MelangeObjectID equal to SearchObjectID */
-	melange::BaseObject* FindMelangeObject(const FString& SearchObjectID, melange::BaseObject* Object);
+	/** Searches the cineware object hierarchy for a cineware::BaseObject that has a MelangeObjectID equal to SearchObjectID */
+	cineware::BaseObject* FindMelangeObject(const FString& SearchObjectID, cineware::BaseObject* Object);
 
 	/** Fetch the object corresponding to the position encoded in HierarchyPosition, starting from Object */
-	melange::BaseObject* GoToMelangeHierarchyPosition(melange::BaseObject* Object, const FString& HierarchyPosition);
+	cineware::BaseObject* GoToMelangeHierarchyPosition(cineware::BaseObject* Object, const FString& HierarchyPosition);
 
-	/** Gets all melange objects that are part of the hierarchy of InstanceRoot. Used to identify child hierarchies of Oinstance objects */
-	//const TArray<melange::BaseObject*>& GetMelangeInstanceObjects(melange::BaseObject* InstanceRoot);
+	/** Gets all cineware objects that are part of the hierarchy of InstanceRoot. Used to identify child hierarchies of Oinstance objects */
+	//const TArray<cineware::BaseObject*>& GetMelangeInstanceObjects(cineware::BaseObject* InstanceRoot);
 
 	/**
 	 * Marks actors children of EmitterObject as ParticleActors, so that they can receive an artificial visibility animation track to
 	 * emulate the look of the particle spawning and despawning
 	 */
-	void MarkActorsAsParticles(melange::BaseObject* EmitterObject, melange::BaseObject* EmittersCache);
+	void MarkActorsAsParticles(cineware::BaseObject* EmitterObject, cineware::BaseObject* EmittersCache);
 
 	/**
-	 * Import melange objects into Datasmith elements. Assets like meshes, materials and textures are added to the Datasmith scene directly,
+	 * Import cineware objects into Datasmith elements. Assets like meshes, materials and textures are added to the Datasmith scene directly,
 	 * while actors are merely returned and must be added as children to scene actors or added to the scene manually.
 	 *
 	 * Returns a null TSharedPtr if an error occurred during the import process.
 	 */
-	TSharedPtr<IDatasmithActorElement> ImportNullActor(melange::BaseObject* Object, const FString& DatasmithName, const FString& DatasmithLabel);
-	TSharedPtr<IDatasmithLightActorElement> ImportLight(melange::BaseObject* LightObject, const FString& DatasmithName, const FString& DatasmithLabel);
-	TSharedPtr<IDatasmithCameraActorElement> ImportCamera(melange::BaseObject* CameraObject, const FString& DatasmithName, const FString& DatasmithLabel);
-	TSharedPtr<IDatasmithMeshActorElement> ImportPolygon(melange::PolygonObject* PolyObject, const FString& DatasmithName, const FString& DatasmithLabel, const TArray<melange::TextureTag*>& TextureTags);
-	TSharedPtr<IDatasmithMaterialInstanceElement> ImportMaterial(melange::Material* InC4DMaterialPtr);
-	TSharedPtr<IDatasmithMaterialInstanceElement> ImportSimpleColorMaterial(melange::BaseObject* Object, int32 UseColor);
+	TSharedPtr<IDatasmithActorElement> ImportNullActor(cineware::BaseObject* Object, const FString& DatasmithName, const FString& DatasmithLabel);
+	TSharedPtr<IDatasmithLightActorElement> ImportLight(cineware::BaseObject* LightObject, const FString& DatasmithName, const FString& DatasmithLabel);
+	TSharedPtr<IDatasmithCameraActorElement> ImportCamera(cineware::BaseObject* CameraObject, const FString& DatasmithName, const FString& DatasmithLabel);
+	TSharedPtr<IDatasmithMeshActorElement> ImportPolygon(cineware::PolygonObject* PolyObject, const FString& DatasmithName, const FString& DatasmithLabel, const TArray<cineware::TextureTag*>& TextureTags);
+	TSharedPtr<IDatasmithMaterialInstanceElement> ImportMaterial(cineware::Material* InC4DMaterialPtr);
+	TSharedPtr<IDatasmithMaterialInstanceElement> ImportSimpleColorMaterial(cineware::BaseObject* Object, int32 UseColor);
 	TSharedPtr<IDatasmithTextureElement> ImportTexture(const FString& TexturePath, EDatasmithTextureMode TextureMode);
-	TSharedPtr<IDatasmithMeshElement> ImportMesh(melange::PolygonObject* PolyObject, const FString& DatasmithMeshName, const FString& DatasmithLabel);
+	TSharedPtr<IDatasmithMeshElement> ImportMesh(cineware::PolygonObject* PolyObject, const FString& DatasmithMeshName, const FString& DatasmithLabel);
 
 	/** Parses the spline and its cache into SplineCurves so that it can be used as paths for animation later */
-	void ImportSpline(melange::SplineObject* ActorObject);
+	void ImportSpline(cineware::SplineObject* ActorObject);
 
-	/** Traverse the melange material hierarchy contained in the c4d file and import each into IDatasmithMaterialInstanceElements */
-	bool ImportMaterialHierarchy(melange::BaseMaterial* InC4DMaterialPtr);
+	/** Traverse the cineware material hierarchy contained in the c4d file and import each into IDatasmithMaterialInstanceElements */
+	bool ImportMaterialHierarchy(cineware::BaseMaterial* InC4DMaterialPtr);
 
-	/** Uses ActorElementToC4DObject to find the corresponding melange object for ActorElement and adds all of its animations to LevelSequence */
+	/** Uses ActorElementToC4DObject to find the corresponding cineware object for ActorElement and adds all of its animations to LevelSequence */
 	void ImportAnimations(TSharedPtr<IDatasmithActorElement> ActorElement);
 
-	/** Traverse the Datasmith scene's IDatasmithActorElement hierarchy and import all animations for the corresponding melange actors*/
+	/** Traverse the Datasmith scene's IDatasmithActorElement hierarchy and import all animations for the corresponding cineware actors*/
 	void ImportActorHierarchyAnimations(TSharedPtr<IDatasmithActorElement> ActorElement);
 
 	/** Searches for the first valid texture used by BaseShader */
-	FString GetBaseShaderTextureFilePath(melange::BaseList2D* BaseShader);
+	FString GetBaseShaderTextureFilePath(cineware::BaseList2D* BaseShader);
 
 	/**
 	 * Generates a new copy of the IDatasmithMaterialInstanceElement with name InMaterialID and alter its properties to match values retrieved
 	 * from InTextureTag, and adds the new material to the Datasmith scene. This is used because texture tags are closer to material
 	 * instances, and may have different "overrides" for each property
 	 */
-	FString CustomizeMaterial(const FString& InMaterialID, const FString& InMeshID, melange::TextureTag* InTextureTag);
+	FString CustomizeMaterial(const FString& InMaterialID, const FString& InMeshID, cineware::TextureTag* InTextureTag);
 
 	/**
 	 * Creates customized materials if necessary, and returns a map from material slot indices to material names
 	 */
-	TMap<int32, FString> GetCustomizedMaterialAssignment(const FString& DatasmithMeshName, const TArray<melange::TextureTag*>& TextureTags);
+	TMap<int32, FString> GetCustomizedMaterialAssignment(const FString& DatasmithMeshName, const TArray<cineware::TextureTag*>& TextureTags);
 
-	/** Imports a melange actor, which might involve parsing another small hierarchy of subnodes and deformers*/
-	TSharedPtr<IDatasmithActorElement> ImportObjectAndChildren(melange::BaseObject* ActorObject, melange::BaseObject* DataObject, TSharedPtr<IDatasmithActorElement> ParentActor, const melange::Matrix& WorldTransformMatrix, const FString& InstancePath, const FString& DatasmithLabel, const TArray<melange::TextureTag*>& TextureTags);
+	/** Imports a cineware actor, which might involve parsing another small hierarchy of subnodes and deformers*/
+	TSharedPtr<IDatasmithActorElement> ImportObjectAndChildren(cineware::BaseObject* ActorObject, cineware::BaseObject* DataObject, TSharedPtr<IDatasmithActorElement> ParentActor, const cineware::Matrix& WorldTransformMatrix, const FString& InstancePath, const FString& DatasmithLabel, const TArray<cineware::TextureTag*>& TextureTags);
 
-	/** Traverse the melange actor hierarchy importing all nodes */
-	void ImportHierarchy(melange::BaseObject* ActorObject, melange::BaseObject* DataObject, TSharedPtr<IDatasmithActorElement> ParentActor, const melange::Matrix& WorldTransformMatrix, const FString& InstancePath, const TArray<melange::TextureTag*>& TextureTags);
+	/** Traverse the cineware actor hierarchy importing all nodes */
+	void ImportHierarchy(cineware::BaseObject* ActorObject, cineware::BaseObject* DataObject, TSharedPtr<IDatasmithActorElement> ParentActor, const cineware::Matrix& WorldTransformMatrix, const FString& InstancePath, const TArray<cineware::TextureTag*>& TextureTags);
 
 	/**
-	 * Adds Actor as a child of ParentActor using the corresponding WorldTransformMatrix. Object is the corresponding melange Object to Actor.
-	 * This will also do the necessary coordinate system conversions between melange and Unreal.
+	 * Adds Actor as a child of ParentActor using the corresponding WorldTransformMatrix. Object is the corresponding cineware Object to Actor.
+	 * This will also do the necessary coordinate system conversions between cineware and Unreal.
 	 */
-	bool AddChildActor(melange::BaseObject* Object, TSharedPtr<IDatasmithActorElement> ParentActor, melange::Matrix WorldTransformMatrix, const TSharedPtr<IDatasmithActorElement>& Actor);
+	bool AddChildActor(cineware::BaseObject* Object, TSharedPtr<IDatasmithActorElement> ParentActor, cineware::Matrix WorldTransformMatrix, const TSharedPtr<IDatasmithActorElement>& Actor);
 
 	/** Event for when a C4D document is about to be opened for translation */
-	DECLARE_EVENT_TwoParams(FDatasmithC4DTranslator, FPreTranslateEvent, melange::BaseDocument*, const FString&)
+	DECLARE_EVENT_TwoParams(FDatasmithC4DTranslator, FPreTranslateEvent, cineware::BaseDocument*, const FString&)
 	static FPreTranslateEvent& OnPreTranslate() { return PreTranslateEvent; }
 
-	melange::BaseDocument* C4dDocument = nullptr;
+	cineware::BaseDocument* C4dDocument = nullptr;
 	FString C4dDocumentFilename;
 
 private:
@@ -153,13 +158,13 @@ private:
 	FVector GetDocumentDefaultColor();
 
 	/** Returns the TextureTags that should affect this object. May check parent objects, so relies on CachesOriginalObject */
-	TArray<melange::TextureTag*> GetActiveTextureTags(const melange::BaseObject* Object, const TArray<melange::TextureTag*>& OrderedTextureTags);
+	TArray<cineware::TextureTag*> GetActiveTextureTags(const cineware::BaseObject* Object, const TArray<cineware::TextureTag*>& OrderedTextureTags);
 
 	/** Removes from Context->Scene all empty actors that have a single child */
 	//void RemoveEmptyActors();
 
 	/** Traverses the original and instanced hierarchy simultaneously and register links between instanced objects and their originals */
-	void RegisterInstancedHierarchy(melange::BaseObject* InstanceRoot, melange::BaseObject* OriginalRoot);
+	void RegisterInstancedHierarchy(cineware::BaseObject* InstanceRoot, cineware::BaseObject* OriginalRoot);
 
 	/**
 	 * Replaces the values of ActorElementToAnimationSources to point to the original objects, in case they are instanced objects.
@@ -172,7 +177,7 @@ private:
 	TMap<IDatasmithMeshElement*, FMeshDescription> MeshElementToMeshDescription;
 
 	/** Storage of imported spline data to be used exclusively for importing animations that follow spline paths */
-	TMap<melange::SplineObject*, TArray<FRichCurve>> SplineCurves;
+	TMap<cineware::SplineObject*, TArray<FRichCurve>> SplineCurves;
 
 	/** Storage of created materials used by CustomizeMaterial to create new "material instances" */
 	TMap<FString, TSharedPtr<IDatasmithMaterialInstanceElement>> MaterialNameToMaterialElement;
@@ -183,20 +188,20 @@ private:
 	/** Cache to prevent us from importing the same texture in the same mode more than once (mode is encoded in the FString as well) */
 	TMap<FString, TSharedPtr<IDatasmithTextureElement>> ImportedTextures;
 
-	/** Storage of all parsed actors from the melange document, used so we can import all animations afterwards */
-	TMap<IDatasmithActorElement*, melange::BaseObject*> ActorElementToAnimationSources;
+	/** Storage of all parsed actors from the cineware document, used so we can import all animations afterwards */
+	TMap<IDatasmithActorElement*, cineware::BaseObject*> ActorElementToAnimationSources;
 
 	/** Maps an instance to the corresponding original node, used so that we can redirect animations to the original nodes */
-	TMap<melange::BaseObject*, melange::BaseObject*> InstancedSubObjectsToOriginals;
+	TMap<cineware::BaseObject*, cineware::BaseObject*> InstancedSubObjectsToOriginals;
 
-	/** Keeps track of the owners of every melange cache object so we can climb the hierarchy upwards */
-	TMap<melange::BaseObject*, melange::BaseObject*> CachesOriginalObject;
+	/** Keeps track of the owners of every cineware cache object so we can climb the hierarchy upwards */
+	TMap<cineware::BaseObject*, cineware::BaseObject*> CachesOriginalObject;
 
 	/** Stores all FCraneCameraAttributes for each camera */
-	TMap<melange::BaseObject*, TSharedRef<FCraneCameraAttributes>> CraneCameraToAttributes;
+	TMap<cineware::BaseObject*, TSharedRef<FCraneCameraAttributes>> CraneCameraToAttributes;
 
 	/** Melange actors that are actually baked 'mesh particles' and need to receive an extra visibility track in ImportAnimations */
-	TSet<melange::BaseObject*> ParticleActors;
+	TSet<cineware::BaseObject*> ParticleActors;
 
 	/** Caches to make sure we don't have any actor name collisions */
 	TSet<FString> NamesOfAllActors;
diff --git a/Engine/Plugins/Enterprise/DatasmithC4DImporter/Source/DatasmithC4DTranslator/Private/DatasmithC4DMelangeSDKEnterGuard.h b/Engine/Plugins/Enterprise/DatasmithC4DImporter/Source/DatasmithC4DTranslator/Private/DatasmithC4DMelangeSDKEnterGuard.h
--- a/Engine/Plugins/Enterprise/DatasmithC4DImporter/Source/DatasmithC4DTranslator/Private/DatasmithC4DMelangeSDKEnterGuard.h
+++ b/Engine/Plugins/Enterprise/DatasmithC4DImporter/Source/DatasmithC4DTranslator/Private/DatasmithC4DMelangeSDKEnterGuard.h
@@ -1,13 +1,26 @@
 // Copyright Epic Games, Inc. All Rights Reserved.
 
 // Enter Datasmith platform include gard.
+#define MAXON_TARGET_RELEASE 1
+#define MAXON_TARGET_64BIT 1
+
 #if PLATFORM_WINDOWS
 #include "Windows/AllowWindowsPlatformTypes.h"
 
+#define MAXON_TARGET_WINDOWS 1
+
 // Back up some preprocessor identifiers.
 __pragma(push_macro("PI"))
 __pragma(push_macro("BYTE_MAX"))
 __pragma(push_macro("_HAS_EXCEPTIONS"))
+__pragma(push_macro("__LP64__"))
+__pragma(push_macro("__APPLE__"))
+
+__pragma(push_macro("GCC34T"))
+#define GCC34T this->
+
+#define __LP64__  0
+#define __APPLE__ 0
 
 #undef PI
 
@@ -15,13 +28,26 @@ __pragma(push_macro("_HAS_EXCEPTIONS"))
 #undef BYTE_MAX
 
 __pragma(warning(push))
-__pragma(warning(disable: 6297)) /* melange\20.004_rbmelange20.0_259890\includes\c4d_drawport.h(276): Arithmetic overflow:  32-bit value is shifted, then cast to 64-bit value.  Results might not be an expected value. */
+__pragma(warning(disable: 6297)) /* cineware\20.004_rbmelange20.0_259890\includes\c4d_drawport.h(276): Arithmetic overflow:  32-bit value is shifted, then cast to 64-bit value.  Results might not be an expected value. */
+
+#elif PLATFORM_MAC
+
+#define MAXON_TARGET_OSX 1
 
 #endif // PLATFORM_WINDOWS
 
 #ifdef __clang__
+
+#pragma push_macro("GCC34T")
+#pragma push_macro("PI")
+
+#define GCC34T this->
+#undef PI
+
 _Pragma( "clang diagnostic push" )
 _Pragma( "clang diagnostic ignored \"-Wdeprecated-declarations\"" )
-#endif
+_Pragma( "clang diagnostic ignored \"-Wundef\"")
+
+#endif // __clang__
 
 
diff --git a/Engine/Plugins/Enterprise/DatasmithC4DImporter/Source/DatasmithC4DTranslator/Private/DatasmithC4DMelangeSDKLeaveGuard.h b/Engine/Plugins/Enterprise/DatasmithC4DImporter/Source/DatasmithC4DTranslator/Private/DatasmithC4DMelangeSDKLeaveGuard.h
--- a/Engine/Plugins/Enterprise/DatasmithC4DImporter/Source/DatasmithC4DTranslator/Private/DatasmithC4DMelangeSDKLeaveGuard.h
+++ b/Engine/Plugins/Enterprise/DatasmithC4DImporter/Source/DatasmithC4DTranslator/Private/DatasmithC4DMelangeSDKLeaveGuard.h
@@ -1,18 +1,38 @@
 // Copyright Epic Games, Inc. All Rights Reserved.
 
 #ifdef __clang__
+
 _Pragma( "clang diagnostic pop" )
-#endif
+
+#pragma pop_macro("GCC34T")
+#pragma pop_macro("PI")
+
+#endif // __clang__
 
 #if PLATFORM_WINDOWS
 
+#undef MAXON_TARGET_WINDOWS
+
 __pragma(warning(pop))
 
 // Restore some preprocessor identifiers.
 __pragma(pop_macro("PI"))
 __pragma(pop_macro("BYTE_MAX"))
 __pragma(pop_macro("_HAS_EXCEPTIONS"))
+__pragma(pop_macro("__LP64__"))
+__pragma(pop_macro("__APPLE__"))
+
+__pragma(pop_macro("GCC34T"))
 
 // Leave Datasmith platform include guard.
 #include "Windows/HideWindowsPlatformTypes.h"
-#endif
+
+#elif PLATFORM_MAC
+
+#undef MAXON_TARGET_OSX
+
+#endif // PLATFORM_WINDOWS
+
+#undef MAXON_TARGET_RELEASE
+#undef MAXON_TARGET_64BIT
+
diff --git a/Engine/Plugins/Enterprise/DatasmithC4DImporter/Source/DatasmithC4DTranslator/Private/DatasmithC4DMelangeSDKOverload.cpp b/Engine/Plugins/Enterprise/DatasmithC4DImporter/Source/DatasmithC4DTranslator/Private/DatasmithC4DMelangeSDKOverload.cpp
--- a/Engine/Plugins/Enterprise/DatasmithC4DImporter/Source/DatasmithC4DTranslator/Private/DatasmithC4DMelangeSDKOverload.cpp
+++ b/Engine/Plugins/Enterprise/DatasmithC4DImporter/Source/DatasmithC4DTranslator/Private/DatasmithC4DMelangeSDKOverload.cpp
@@ -9,7 +9,7 @@
 #include "default_alien_overloads.h"
 #include "DatasmithC4DMelangeSDKLeaveGuard.h"
 
-namespace melange
+namespace cineware
 {
 #pragma pack (push, 8)
 
@@ -55,7 +55,7 @@ namespace melange
 }
 
 // overload this function and fill in your own unique data
-void GetWriterInfo(melange::Int32 &id, melange::String &appname)
+void GetWriterInfo(cineware::Int32 &id, cineware::String &appname)
 {
 	// register your own pluginid once for your exporter and enter it here under id
 	// this id must be used for your own unique ids
diff --git a/Engine/Plugins/Enterprise/DatasmithC4DImporter/Source/DatasmithC4DTranslator/Private/DatasmithC4DTranslator.h b/Engine/Plugins/Enterprise/DatasmithC4DImporter/Source/DatasmithC4DTranslator/Private/DatasmithC4DTranslator.h
--- a/Engine/Plugins/Enterprise/DatasmithC4DImporter/Source/DatasmithC4DTranslator/Private/DatasmithC4DTranslator.h
+++ b/Engine/Plugins/Enterprise/DatasmithC4DImporter/Source/DatasmithC4DTranslator/Private/DatasmithC4DTranslator.h
@@ -29,7 +29,7 @@ public:
 	virtual void SetSceneImportOptions(TArray<TStrongObjectPtr<UDatasmithOptionsBase>>& Options) override;
 
 	/** Event for when a C4D document is about to be opened for translation */
-	/*DECLARE_EVENT_TwoParams(FDatasmithC4DTranslator, FPreTranslateEvent, melange::BaseDocument*, const FString&)
+	/*DECLARE_EVENT_TwoParams(FDatasmithC4DTranslator, FPreTranslateEvent, cineware::BaseDocument*, const FString&)
 	static FPreTranslateEvent& OnPreTranslate() { return PreTranslateEvent; }*/
 	
 private:
diff --git a/Engine/Plugins/Enterprise/DatasmithC4DImporter/Source/DatasmithC4DTranslator/Private/DatasmithC4DUtils.cpp b/Engine/Plugins/Enterprise/DatasmithC4DImporter/Source/DatasmithC4DTranslator/Private/DatasmithC4DUtils.cpp
--- a/Engine/Plugins/Enterprise/DatasmithC4DImporter/Source/DatasmithC4DTranslator/Private/DatasmithC4DUtils.cpp
+++ b/Engine/Plugins/Enterprise/DatasmithC4DImporter/Source/DatasmithC4DTranslator/Private/DatasmithC4DUtils.cpp
@@ -14,12 +14,12 @@
 #include "Misc/Paths.h"
 #include "Misc/SecureHash.h"
 
-int32 MelangeGetInt32(melange::BaseList2D* Object, melange::Int32 Parameter)
+int32 MelangeGetInt32(cineware::BaseList2D* Object, cineware::Int32 Parameter)
 {
 	if (Object)
 	{
-		melange::GeData Data;
-		if (Object->GetParameter(Parameter, Data) && Data.GetType() == melange::DA_LONG)
+		cineware::GeData Data;
+		if (Object->GetParameter(Parameter, Data) && Data.GetType() == cineware::DA_LONG)
 		{
 			return static_cast<int32>(Data.GetInt32());
 		}
@@ -27,12 +27,12 @@ int32 MelangeGetInt32(melange::BaseList2D* Object, melange::Int32 Parameter)
 	return 0;
 }
 
-int64 MelangeGetInt64(melange::BaseList2D* Object, melange::Int32 Parameter)
+int64 MelangeGetInt64(cineware::BaseList2D* Object, cineware::Int32 Parameter)
 {
 	if (Object)
 	{
-		melange::GeData Data;
-		if (Object->GetParameter(Parameter, Data) && Data.GetType() == melange::DA_LLONG)
+		cineware::GeData Data;
+		if (Object->GetParameter(Parameter, Data) && Data.GetType() == cineware::DA_LLONG)
 		{
 			return static_cast<int64>(Data.GetInt64());
 		}
@@ -40,12 +40,12 @@ int64 MelangeGetInt64(melange::BaseList2D* Object, melange::Int32 Parameter)
 	return 0;
 }
 
-bool MelangeGetBool(melange::BaseList2D* Object, melange::Int32 Parameter)
+bool MelangeGetBool(cineware::BaseList2D* Object, cineware::Int32 Parameter)
 {
 	if (Object)
 	{
-		melange::GeData Data;
-		if (Object->GetParameter(Parameter, Data) && Data.GetType() == melange::DA_LONG)
+		cineware::GeData Data;
+		if (Object->GetParameter(Parameter, Data) && Data.GetType() == cineware::DA_LONG)
 		{
 			return Data.GetBool();
 		}
@@ -53,18 +53,18 @@ bool MelangeGetBool(melange::BaseList2D* Object, melange::Int32 Parameter)
 	return false;
 }
 
-double MelangeGetDouble(melange::BaseList2D* Object, melange::Int32 Parameter)
+double MelangeGetDouble(cineware::BaseList2D* Object, cineware::Int32 Parameter)
 {
 	if (Object)
 	{
-		melange::GeData Data;
+		cineware::GeData Data;
 		if (Object->GetParameter(Parameter, Data))
 		{
-			if(Data.GetType() == melange::DA_REAL)
+			if(Data.GetType() == cineware::DA_REAL)
 			{
 				return static_cast<double>(Data.GetFloat());
 			}
-			else if (Data.GetType() == melange::DA_TIME)
+			else if (Data.GetType() == cineware::DA_TIME)
 			{
 				return static_cast<double>(Data.GetTime().Get());
 			}
@@ -73,12 +73,12 @@ double MelangeGetDouble(melange::BaseList2D* Object, melange::Int32 Parameter)
 	return 0.0;
 }
 
-FVector MelangeGetVector(melange::BaseList2D* Object, melange::Int32 Parameter)
+FVector MelangeGetVector(cineware::BaseList2D* Object, cineware::Int32 Parameter)
 {
 	if (Object)
 	{
-		melange::GeData Data;
-		if (Object->GetParameter(Parameter, Data) && Data.GetType() == melange::DA_VECTOR)
+		cineware::GeData Data;
+		if (Object->GetParameter(Parameter, Data) && Data.GetType() == cineware::DA_VECTOR)
 		{
 			return MelangeVectorToFVector(Data.GetVector());
 		}
@@ -86,12 +86,12 @@ FVector MelangeGetVector(melange::BaseList2D* Object, melange::Int32 Parameter)
 	return FVector::ZeroVector;
 }
 
-FMatrix MelangeGetMatrix(melange::BaseList2D* Object, melange::Int32 Parameter)
+FMatrix MelangeGetMatrix(cineware::BaseList2D* Object, cineware::Int32 Parameter)
 {
 	if (Object)
 	{
-		melange::GeData Data;
-		if (Object->GetParameter(Parameter, Data) && Data.GetType() == melange::DA_MATRIX)
+		cineware::GeData Data;
+		if (Object->GetParameter(Parameter, Data) && Data.GetType() == cineware::DA_MATRIX)
 		{
 			return MelangeMatrixToFMatrix(Data.GetMatrix());
 		}
@@ -102,15 +102,15 @@ FMatrix MelangeGetMatrix(melange::BaseList2D* Object, melange::Int32 Parameter)
 	return Result;
 }
 
-TArray<uint8> MelangeGetByteArray(melange::BaseList2D* Object, melange::Int32 Parameter)
+TArray<uint8> MelangeGetByteArray(cineware::BaseList2D* Object, cineware::Int32 Parameter)
 {
 	TArray<uint8> Result;
 	if (Object)
 	{
-		melange::GeData Data;
-		if (Object->GetParameter(Parameter, Data) && Data.GetType() == melange::DA_BYTEARRAY)
+		cineware::GeData Data;
+		if (Object->GetParameter(Parameter, Data) && Data.GetType() == cineware::DA_BYTEARRAY)
 		{
-			const melange::ByteArray& Arr = Data.GetByteArray();
+			const cineware::ByteArray& Arr = Data.GetByteArray();
 			Result.SetNumUninitialized(static_cast<int32>(Arr.size));
 			FMemory::Memcpy(Result.GetData(), Arr.mem, Arr.size);
 		}
@@ -118,18 +118,18 @@ TArray<uint8> MelangeGetByteArray(melange::BaseList2D* Object, melange::Int32 Pa
 	return Result;
 }
 
-FString MelangeGetString(melange::BaseList2D* Object, melange::Int32 Parameter)
+FString MelangeGetString(cineware::BaseList2D* Object, cineware::Int32 Parameter)
 {
 	if (Object)
 	{
-		melange::GeData Data;
+		cineware::GeData Data;
 		if (Object->GetParameter(Parameter, Data))
 		{
-			if(Data.GetType() == melange::DA_STRING)
+			if(Data.GetType() == cineware::DA_STRING)
 			{
 				return MelangeStringToFString(Data.GetString());
 			}
-			else if (Data.GetType() == melange::DA_FILENAME)
+			else if (Data.GetType() == cineware::DA_FILENAME)
 			{
 				return MelangeFilenameToPath(Data.GetFilename());
 			}
@@ -138,12 +138,12 @@ FString MelangeGetString(melange::BaseList2D* Object, melange::Int32 Parameter)
 	return FString();
 }
 
-melange::BaseList2D* MelangeGetLink(melange::BaseList2D* Object, melange::Int32 Parameter)
+cineware::BaseList2D* MelangeGetLink(cineware::BaseList2D* Object, cineware::Int32 Parameter)
 {
 	if (Object)
 	{
-		melange::GeData Data;
-		if (Object->GetParameter(Parameter, Data) && Data.GetType() == melange::DA_ALIASLINK)
+		cineware::GeData Data;
+		if (Object->GetParameter(Parameter, Data) && Data.GetType() == cineware::DA_ALIASLINK)
 		{
 			return Data.GetLink();
 		}
@@ -151,17 +151,17 @@ melange::BaseList2D* MelangeGetLink(melange::BaseList2D* Object, melange::Int32
 	return nullptr;
 }
 
-float MelangeGetFloat(melange::BaseList2D* Object, melange::Int32 Parameter)
+float MelangeGetFloat(cineware::BaseList2D* Object, cineware::Int32 Parameter)
 {
 	return static_cast<float>(MelangeGetDouble(Object, Parameter));
 }
 
-FVector ConvertMelangePosition(const melange::Vector32& MelangePosition, float WorldUnitScale)
+FVector ConvertMelangePosition(const cineware::Vector32& MelangePosition, float WorldUnitScale)
 {
 	return ConvertMelangePosition(MelangeVectorToFVector(MelangePosition), WorldUnitScale);
 }
 
-FVector ConvertMelangePosition(const melange::Vector64& MelangePosition, float WorldUnitScale)
+FVector ConvertMelangePosition(const cineware::Vector64& MelangePosition, float WorldUnitScale)
 {
 	return ConvertMelangePosition(MelangeVectorToFVector(MelangePosition), WorldUnitScale);
 }
@@ -171,12 +171,12 @@ FVector ConvertMelangePosition(const FVector& MelangePosition, float WorldUnitSc
 	return WorldUnitScale * FVector(MelangePosition.X, -MelangePosition.Z, MelangePosition.Y);
 }
 
-FVector ConvertMelangeDirection(const melange::Vector32& MelangePosition)
+FVector ConvertMelangeDirection(const cineware::Vector32& MelangePosition)
 {
 	return ConvertMelangeDirection(MelangeVectorToFVector(MelangePosition));
 }
 
-FVector ConvertMelangeDirection(const melange::Vector64& MelangePosition)
+FVector ConvertMelangeDirection(const cineware::Vector64& MelangePosition)
 {
 	return ConvertMelangeDirection(MelangeVectorToFVector(MelangePosition));
 }
@@ -186,24 +186,24 @@ FVector ConvertMelangeDirection(const FVector& MelangePosition)
 	return FVector(MelangePosition.X, -MelangePosition.Z, MelangePosition.Y);
 }
 
-FVector MelangeVectorToFVector(const melange::Vector32& MelangeVector)
+FVector MelangeVectorToFVector(const cineware::Vector32& MelangeVector)
 {
 	return FVector(MelangeVector.x, MelangeVector.y, MelangeVector.z);
 }
 
-FVector MelangeVectorToFVector(const melange::Vector64& MelangeVector)
+FVector MelangeVectorToFVector(const cineware::Vector64& MelangeVector)
 {
 	return FVector(static_cast<float>(MelangeVector.x),
 				   static_cast<float>(MelangeVector.y),
 		           static_cast<float>(MelangeVector.z));
 }
 
-FVector4 MelangeVector4ToFVector4(const melange::Vector4d32& MelangeVector)
+FVector4 MelangeVector4ToFVector4(const cineware::Vector4d32& MelangeVector)
 {
 	return FVector4(MelangeVector.x, MelangeVector.y, MelangeVector.z, MelangeVector.w);
 }
 
-FVector4 MelangeVector4ToFVector4(const melange::Vector4d64& MelangeVector)
+FVector4 MelangeVector4ToFVector4(const cineware::Vector4d64& MelangeVector)
 {
 	return FVector4(static_cast<float>(MelangeVector.x),
 					static_cast<float>(MelangeVector.y),
@@ -211,7 +211,7 @@ FVector4 MelangeVector4ToFVector4(const melange::Vector4d64& MelangeVector)
 					static_cast<float>(MelangeVector.w));
 }
 
-FMatrix MelangeMatrixToFMatrix(const melange::Matrix& Matrix)
+FMatrix MelangeMatrixToFMatrix(const cineware::Matrix& Matrix)
 {
 	return FMatrix(MelangeVectorToFVector(Matrix.v1),
 				   MelangeVectorToFVector(Matrix.v2),
@@ -219,9 +219,9 @@ FMatrix MelangeMatrixToFMatrix(const melange::Matrix& Matrix)
 				   MelangeVectorToFVector(Matrix.off));
 }
 
-FString MelangeStringToFString(const melange::String& MelangeString)
+FString MelangeStringToFString(const cineware::String& MelangeString)
 {
-	TUniquePtr<melange::Char> CStr(MelangeString.GetCStringCopy());
+	TUniquePtr<cineware::Char> CStr(MelangeString.GetCStringCopy());
 	return CStr.Get();
 }
 
@@ -234,7 +234,7 @@ FString MD5FromString(const FString& Value)
 	return BytesToHex(MD5Hash, 16);
 }
 
-FString MelangeFilenameToPath(const melange::Filename& Filename)
+FString MelangeFilenameToPath(const cineware::Filename& Filename)
 {
 	return MelangeStringToFString(Filename.GetString());
 }
@@ -261,7 +261,7 @@ FString SearchForFile(FString Filename, const FString& C4dDocumentFilename)
 		}
 	}
 
-	// Maybe its a file that has been physically moved to the exported folder, but melange still
+	// Maybe its a file that has been physically moved to the exported folder, but cineware still
 	// has its original filepath
 	FString CleanFilename = FPaths::GetCleanFilename(Filename);
 	FString LocalPath = FPaths::Combine(DocumentPath, CleanFilename);
@@ -270,7 +270,7 @@ FString SearchForFile(FString Filename, const FString& C4dDocumentFilename)
 		return LocalPath;
 	}
 
-	// Try searching inside a 'tex' folder first (where melange emits textures)
+	// Try searching inside a 'tex' folder first (where cineware emits textures)
 	FString PathInTex = FPaths::Combine(DocumentPath, FString("tex"), CleanFilename);
 	if (FPaths::FileExists(PathInTex))
 	{
@@ -288,7 +288,7 @@ FString SearchForFile(FString Filename, const FString& C4dDocumentFilename)
 	return "";
 }
 
-FString MelangeObjectName(melange::BaseList2D* Object)
+FString MelangeObjectName(cineware::BaseList2D* Object)
 {
 	if (!Object)
 	{
@@ -298,69 +298,69 @@ FString MelangeObjectName(melange::BaseList2D* Object)
 	return MelangeStringToFString(Object->GetName());
 }
 
-FString MelangeObjectTypeName(melange::BaseList2D* Object)
+FString MelangeObjectTypeName(cineware::BaseList2D* Object)
 {
 	if (!Object)
 	{
 		return FString(TEXT("Invalid object"));
 	}
 
-	return MelangeStringToFString(melange::String(melange::GetObjectTypeName(Object->GetType())));
+	return MelangeStringToFString(cineware::String(cineware::GetObjectTypeName(Object->GetType())));
 }
 
-FString GeDataToString(const melange::GeData& Data)
+FString GeDataToString(const cineware::GeData& Data)
 {
 	switch (Data.GetType())
 	{
-	case melange::DA_NIL:
+	case cineware::DA_NIL:
 		return TEXT("NIL");
 		break;
-	case melange::DA_VOID:
+	case cineware::DA_VOID:
 		return TEXT("VOID");
 		break;
-	case melange::DA_LONG:
+	case cineware::DA_LONG:
 		return FString::Printf(TEXT("%d"), Data.GetInt32());
 		break;
-	case melange::DA_REAL:
+	case cineware::DA_REAL:
 		return FString::Printf(TEXT("%f"), Data.GetFloat());
 		break;
-	case melange::DA_TIME:
+	case cineware::DA_TIME:
 		return FString::Printf(TEXT("%f"), Data.GetTime().Get());
 		break;
-	case melange::DA_VECTOR:
+	case cineware::DA_VECTOR:
 		return MelangeVectorToFVector(Data.GetVector()).ToString();
 		break;
-	case melange::DA_MATRIX:
+	case cineware::DA_MATRIX:
 		return MelangeMatrixToFMatrix(Data.GetMatrix()).ToString();
 		break;
-	case melange::DA_LLONG:
+	case cineware::DA_LLONG:
 		return FString::Printf(TEXT("%jd"), Data.GetInt64());
 		break;
-	case melange::DA_BYTEARRAY:
+	case cineware::DA_BYTEARRAY:
 	{
-		const melange::ByteArray& Arr = Data.GetByteArray();
+		const cineware::ByteArray& Arr = Data.GetByteArray();
 		return BytesToHex(static_cast<const uint8*>(Arr.mem), static_cast<int32>(Arr.size));
 		break;
 	}
-	case melange::DA_STRING:
+	case cineware::DA_STRING:
 		return MelangeStringToFString(Data.GetString());
 		break;
-	case melange::DA_FILENAME:
+	case cineware::DA_FILENAME:
 		return MelangeFilenameToPath(Data.GetFilename());
 		break;
-	case melange::DA_CONTAINER:
+	case cineware::DA_CONTAINER:
 		return FString::Printf(TEXT("%d"), Data.GetContainer()->GetId());
 		break;
-	case melange::DA_ALIASLINK:
+	case cineware::DA_ALIASLINK:
 	{
-		melange::BaseList2D* Target = Data.GetLink();
+		cineware::BaseList2D* Target = Data.GetLink();
 		return  FString::Printf(TEXT("%s @0x%x"), *MelangeObjectName(Target), Target);
 		break;
 	}
-	case melange::DA_MARKER:
+	case cineware::DA_MARKER:
 		return TEXT("MARKER");
 		break;
-	case melange::DA_MISSINGPLUG:
+	case cineware::DA_MISSINGPLUG:
 		return TEXT("MISSINGPLUG");
 		break;
 	default:
@@ -375,21 +375,21 @@ FString GeTypeToString(int32 GeType)
 	using GeTypeEntry = TPairInitializer<const int32&, const FString&>;
 	const static TMap<int32, FString> GeTypes
 	({
-		GeTypeEntry(melange::DA_NIL, TEXT("DA_NIL (no value)")),
-		GeTypeEntry(melange::DA_VOID, TEXT("DA_VOID (void pointer)")),
-		GeTypeEntry(melange::DA_LONG, TEXT("DA_LONG (int32)")),
-		GeTypeEntry(melange::DA_REAL, TEXT("DA_REAL (double)")),
-		GeTypeEntry(melange::DA_TIME, TEXT("DA_TIME (double)")),
-		GeTypeEntry(melange::DA_VECTOR, TEXT("DA_VECTOR (melange::Vector)")),
-		GeTypeEntry(melange::DA_MATRIX, TEXT("DA_MATRIX (melange::Matrix)")),
-		GeTypeEntry(melange::DA_LLONG, TEXT("DA_LLONG (int64)")),
-		GeTypeEntry(melange::DA_BYTEARRAY, TEXT("DA_BYTEARRAY (void pointer)")),
-		GeTypeEntry(melange::DA_STRING, TEXT("DA_STRING (melange::String)")),
-		GeTypeEntry(melange::DA_FILENAME, TEXT("DA_FILENAME (melange::Filename)")),
-		GeTypeEntry(melange::DA_CONTAINER, TEXT("DA_CONTAINER (melange::BaseContainer)")),
-		GeTypeEntry(melange::DA_ALIASLINK, TEXT("DA_ALIASLINK (melange::BaseLink)")),
-		GeTypeEntry(melange::DA_MARKER, TEXT("DA_MARKER (not used)")),
-		GeTypeEntry(melange::DA_MISSINGPLUG, TEXT("DA_MISSINGPLUG (missing datatype plugin)"))
+		GeTypeEntry(cineware::DA_NIL, TEXT("DA_NIL (no value)")),
+		GeTypeEntry(cineware::DA_VOID, TEXT("DA_VOID (void pointer)")),
+		GeTypeEntry(cineware::DA_LONG, TEXT("DA_LONG (int32)")),
+		GeTypeEntry(cineware::DA_REAL, TEXT("DA_REAL (double)")),
+		GeTypeEntry(cineware::DA_TIME, TEXT("DA_TIME (double)")),
+		GeTypeEntry(cineware::DA_VECTOR, TEXT("DA_VECTOR (cineware::Vector)")),
+		GeTypeEntry(cineware::DA_MATRIX, TEXT("DA_MATRIX (cineware::Matrix)")),
+		GeTypeEntry(cineware::DA_LLONG, TEXT("DA_LLONG (int64)")),
+		GeTypeEntry(cineware::DA_BYTEARRAY, TEXT("DA_BYTEARRAY (void pointer)")),
+		GeTypeEntry(cineware::DA_STRING, TEXT("DA_STRING (cineware::String)")),
+		GeTypeEntry(cineware::DA_FILENAME, TEXT("DA_FILENAME (cineware::Filename)")),
+		GeTypeEntry(cineware::DA_CONTAINER, TEXT("DA_CONTAINER (cineware::BaseContainer)")),
+		GeTypeEntry(cineware::DA_ALIASLINK, TEXT("DA_ALIASLINK (cineware::BaseLink)")),
+		GeTypeEntry(cineware::DA_MARKER, TEXT("DA_MARKER (not used)")),
+		GeTypeEntry(cineware::DA_MISSINGPLUG, TEXT("DA_MISSINGPLUG (missing datatype plugin)"))
 		});
 
 	if (const FString* FoundTypeName = GeTypes.Find(GeType))
@@ -405,11 +405,11 @@ FString GeTypeToString(int32 GeType)
 	return TEXT("UNKNOWN_TYPE (?)");
 }
 
-FString MelangeParameterValueToString(melange::BaseList2D* Object, melange::Int32 ParameterID)
+FString MelangeParameterValueToString(cineware::BaseList2D* Object, cineware::Int32 ParameterID)
 {
 	if (Object)
 	{
-		melange::GeData Data;
+		cineware::GeData Data;
 		if (Object->GetParameter(ParameterID, Data))
 		{
 			return GeDataToString(Data);
@@ -418,16 +418,16 @@ FString MelangeParameterValueToString(melange::BaseList2D* Object, melange::Int3
 	return TEXT("");
 }
 
-TOptional<FString> GetMelangeBaseList2dID(melange::BaseList2D* BaseList)
+TOptional<FString> GetMelangeBaseList2dID(cineware::BaseList2D* BaseList)
 {
 	TOptional<FString> MelangeID;
 	if (BaseList)
 	{
 		if (BaseList->GetUniqueIDCount() > 0)
 		{
-			melange::Int32 AppId;
-			const melange::Char* IdData;
-			melange::Int Bytes;
+			cineware::Int32 AppId;
+			const cineware::Char* IdData;
+			cineware::Int Bytes;
 			if (BaseList->GetUniqueIDIndex(0, AppId, IdData, Bytes))
 			{
 				MelangeID.Emplace(BytesToHex(reinterpret_cast<const uint8*>(&AppId), sizeof(AppId)) + "_" + BytesToHex(reinterpret_cast<const uint8*>(IdData), static_cast<int32>(Bytes)));
diff --git a/Engine/Plugins/Enterprise/DatasmithC4DImporter/Source/DatasmithC4DTranslator/Private/DatasmithC4DUtils.h b/Engine/Plugins/Enterprise/DatasmithC4DImporter/Source/DatasmithC4DTranslator/Private/DatasmithC4DUtils.h
--- a/Engine/Plugins/Enterprise/DatasmithC4DImporter/Source/DatasmithC4DTranslator/Private/DatasmithC4DUtils.h
+++ b/Engine/Plugins/Enterprise/DatasmithC4DImporter/Source/DatasmithC4DTranslator/Private/DatasmithC4DUtils.h
@@ -11,105 +11,105 @@
 #include "DatasmithC4DMelangeSDKLeaveGuard.h"
 
 /**
- * Retrieves the value of a DA_LONG parameter of a melange object as a melange::Int32,
+ * Retrieves the value of a DA_LONG parameter of a cineware object as a cineware::Int32,
  * and static_casts it to an int32. Will return 0 if the parameter is invalid
  */
-int32 MelangeGetInt32(melange::BaseList2D* Object, melange::Int32 Parameter);
+int32 MelangeGetInt32(cineware::BaseList2D* Object, cineware::Int32 Parameter);
 
 /**
- * Retrieves the value of a DA_LLONG parameter of a melange object
+ * Retrieves the value of a DA_LLONG parameter of a cineware object
  * and returns it as an int64. Will return 0 if the parameter is invalid
  */
-int64 MelangeGetInt64(melange::BaseList2D* Object, melange::Int32 Parameter);
+int64 MelangeGetInt64(cineware::BaseList2D* Object, cineware::Int32 Parameter);
 
 /**
- * Retrieves the value of a DA_LONG parameter of a melange object as a melange::Bool,
+ * Retrieves the value of a DA_LONG parameter of a cineware object as a cineware::Bool,
  * and returns it as a bool. Will return false if the parameter is invalid
  */
-bool MelangeGetBool(melange::BaseList2D* Object, melange::Int32 Parameter);
+bool MelangeGetBool(cineware::BaseList2D* Object, cineware::Int32 Parameter);
 
 /**
- * Retrieves the value of a DA_REAL parameter of a melange object
+ * Retrieves the value of a DA_REAL parameter of a cineware object
  * and returns it as a double. Will return 0.0 if the parameter is invalid
  */
-double MelangeGetDouble(melange::BaseList2D* Object, melange::Int32 Parameter);
+double MelangeGetDouble(cineware::BaseList2D* Object, cineware::Int32 Parameter);
 
 /**
- * Retrieves the value of a DA_VECTOR parameter of a melange object
+ * Retrieves the value of a DA_VECTOR parameter of a cineware object
  * and returns it as an FVector. No coordinate or color conversions are applied.
  * Will return FVector() if the parameter is invalid
  */
-FVector MelangeGetVector(melange::BaseList2D* Object, melange::Int32 Parameter);
+FVector MelangeGetVector(cineware::BaseList2D* Object, cineware::Int32 Parameter);
 
 /**
- * Retrieves the value of a DA_MATRIX parameter of a melange object
+ * Retrieves the value of a DA_MATRIX parameter of a cineware object
  * and returns it as an FMatrix. Will return an identity matrix if the parameter is invalid
  */
-FMatrix MelangeGetMatrix(melange::BaseList2D* Object, melange::Int32 Parameter);
+FMatrix MelangeGetMatrix(cineware::BaseList2D* Object, cineware::Int32 Parameter);
 
 /**
- * Retrieves the value of a DA_BYTEARRAY parameter of a melange object
+ * Retrieves the value of a DA_BYTEARRAY parameter of a cineware object
  * and returns it as an array of bytes. Will return TArray<uint8>() if the parameter is invalid
  */
-TArray<uint8> MelangeGetByteArray(melange::BaseList2D* Object, melange::Int32 Parameter);
+TArray<uint8> MelangeGetByteArray(cineware::BaseList2D* Object, cineware::Int32 Parameter);
 
 /**
- * Retrieves the value of a DA_STRING or DA_FILENAME parameter of a melange object
+ * Retrieves the value of a DA_STRING or DA_FILENAME parameter of a cineware object
  * and returns it as a string. Will return FString() if the parameter is invalid
  */
-FString MelangeGetString(melange::BaseList2D* Object, melange::Int32 Parameter);
+FString MelangeGetString(cineware::BaseList2D* Object, cineware::Int32 Parameter);
 
 /**
- * Retrieves a the object pointed to by an DA_ALIASLINK parameter of a melange object
+ * Retrieves a the object pointed to by an DA_ALIASLINK parameter of a cineware object
  * and returns it. Will return nullptr if the parameter is invalid
  */
-melange::BaseList2D* MelangeGetLink(melange::BaseList2D* Object, melange::Int32 Parameter);
+cineware::BaseList2D* MelangeGetLink(cineware::BaseList2D* Object, cineware::Int32 Parameter);
 
 /**
- * Retrieves the value of a DA_REAL (double) parameter from a melange object and static_casts it to float.
+ * Retrieves the value of a DA_REAL (double) parameter from a cineware object and static_casts it to float.
  * Will return 0.0f if the parameter is invalid
  */
-float MelangeGetFloat(melange::BaseList2D* Object, melange::Int32 Parameter);
+float MelangeGetFloat(cineware::BaseList2D* Object, cineware::Int32 Parameter);
 
 /**
- * Converts a vector representing a position from the melange coordinate system to the UnrealEditor coordinate system
+ * Converts a vector representing a position from the cineware coordinate system to the UnrealEditor coordinate system
  */
-FVector ConvertMelangePosition(const melange::Vector32& MelangePosition, float WorldUnitScale = 1.0f);
-FVector ConvertMelangePosition(const melange::Vector64& MelangePosition, float WorldUnitScale = 1.0f);
+FVector ConvertMelangePosition(const cineware::Vector32& MelangePosition, float WorldUnitScale = 1.0f);
+FVector ConvertMelangePosition(const cineware::Vector64& MelangePosition, float WorldUnitScale = 1.0f);
 FVector ConvertMelangePosition(const FVector& MelangePosition, float WorldUnitScale = 1.0f);
 
 /**
- * Converts a vector representing a direction from the melange coordinate system to the UnrealEditor coordinate system
+ * Converts a vector representing a direction from the cineware coordinate system to the UnrealEditor coordinate system
  */
-FVector ConvertMelangeDirection(const melange::Vector32& MelangePosition);
-FVector ConvertMelangeDirection(const melange::Vector64& MelangePosition);
+FVector ConvertMelangeDirection(const cineware::Vector32& MelangePosition);
+FVector ConvertMelangeDirection(const cineware::Vector64& MelangePosition);
 FVector ConvertMelangeDirection(const FVector& MelangePosition);
 
 /**
- * Converts a melange::Vector into an FVector
+ * Converts a cineware::Vector into an FVector
  */
-FVector MelangeVectorToFVector(const melange::Vector32& MelangeVector);
-FVector MelangeVectorToFVector(const melange::Vector64& MelangeVector);
+FVector MelangeVectorToFVector(const cineware::Vector32& MelangeVector);
+FVector MelangeVectorToFVector(const cineware::Vector64& MelangeVector);
 
 /**
- * Converts a melange::Vector4d32 into an FVector4
+ * Converts a cineware::Vector4d32 into an FVector4
  */
-FVector4 MelangeVector4ToFVector4(const melange::Vector4d32& MelangeVector);
+FVector4 MelangeVector4ToFVector4(const cineware::Vector4d32& MelangeVector);
 
 /**
- * Converts a melange::Vector4d64 into an FVector4
+ * Converts a cineware::Vector4d64 into an FVector4
  */
-FVector4 MelangeVector4ToFVector4(const melange::Vector4d64& MelangeVector);
+FVector4 MelangeVector4ToFVector4(const cineware::Vector4d64& MelangeVector);
 
 /**
- * Converts a melange::Matrix into an FMatrix
+ * Converts a cineware::Matrix into an FMatrix
  */
-FMatrix MelangeMatrixToFMatrix(const melange::Matrix& Matrix);
+FMatrix MelangeMatrixToFMatrix(const cineware::Matrix& Matrix);
 
 /**
- * Converts a melange::String into an FString
+ * Converts a cineware::String into an FString
  */
-FString MelangeStringToFString(const melange::String& MelangeString);
+FString MelangeStringToFString(const cineware::String& MelangeString);
 
 /**
  * Uses the bytes of 'String' to generate an MD5 hash, and returns the bytes of the hash back as an FString
@@ -117,14 +117,14 @@ FString MelangeStringToFString(const melange::String& MelangeString);
 FString MD5FromString(const FString& String);
 
 /**
- * Converts a melange::Filename to an FString path
+ * Converts a cineware::Filename to an FString path
  */
-FString MelangeFilenameToPath(const melange::Filename& Filename);
+FString MelangeFilenameToPath(const cineware::Filename& Filename);
 
 /**
- * Searches for a file in places where melange is likely to put them. Will try things like making the filename
+ * Searches for a file in places where cineware is likely to put them. Will try things like making the filename
  * relative to the C4dDocumentFilename's folder, searching the 'tex' folder, and etc. Use this to search for texture
- * and IES files as the path contained in melange is likely to be pointing somewhere else if the user didn't export
+ * and IES files as the path contained in cineware is likely to be pointing somewhere else if the user didn't export
  * the scene just right.
  *
  * @param Filename				The path to the file we're trying to find. Can be relative or absolute
@@ -134,24 +134,24 @@ FString MelangeFilenameToPath(const melange::Filename& Filename);
 FString SearchForFile(FString Filename, const FString& C4dDocumentFilename);
 
 /**
- * Gets the name of a melange object as an FString, or returns 'Invalid object' if the argument is nullptr
+ * Gets the name of a cineware object as an FString, or returns 'Invalid object' if the argument is nullptr
  */
-FString MelangeObjectName(melange::BaseList2D* Object);
+FString MelangeObjectName(cineware::BaseList2D* Object);
 
 /**
- * Gets the type of a melange object as an FString (e.g. Ocloner, Onull, Opolygon), or 'Invalid object' if the
+ * Gets the type of a cineware object as an FString (e.g. Ocloner, Onull, Opolygon), or 'Invalid object' if the
  * argument is nullptr
  */
-FString MelangeObjectTypeName(melange::BaseList2D* Object);
+FString MelangeObjectTypeName(cineware::BaseList2D* Object);
 
 /**
- * Gets the data stored within a melange::GeData object according to its .GetType() and returns that as a string.
+ * Gets the data stored within a cineware::GeData object according to its .GetType() and returns that as a string.
  * Will try to convert to Unreal types first (like FVector or FMatrix) and convert those to string instead
  */
-FString GeDataToString(const melange::GeData& Data);
+FString GeDataToString(const cineware::GeData& Data);
 
 /**
- * Converts the type returned by .GetType() of a melange::GeData instance from an melange::Int32 into
+ * Converts the type returned by .GetType() of a cineware::GeData instance from an cineware::Int32 into
  * a human-readable string
  */
 FString GeTypeToString(int32 GeType);
@@ -160,11 +160,11 @@ FString GeTypeToString(int32 GeType);
  * Gets the corresponding parameter from Object, and returns the value stored in the resulting GeData as a string.
  * Will return an empty string if Object is nullptr, but it can return NIL if the parameter is not found
  */
-FString MelangeParameterValueToString(melange::BaseList2D* Object, melange::Int32 ParameterID);
+FString MelangeParameterValueToString(cineware::BaseList2D* Object, cineware::Int32 ParameterID);
 
 /**
- * Returns the full melange ID for the BaseList2D argument as an FString, which will include AppId
+ * Returns the full cineware ID for the BaseList2D argument as an FString, which will include AppId
  */
-TOptional<FString> GetMelangeBaseList2dID(melange::BaseList2D* BaseList);
+TOptional<FString> GetMelangeBaseList2dID(cineware::BaseList2D* BaseList);
 
 #endif
\ No newline at end of file
