Fix Clang StaticAnalysis warns C++ Object Pointers are null
Fix Clang Static Analysis warns Dereference of a null pointer

diff --git a/Engine/Plugins/Enterprise/DatasmithC4DImporter/Source/DatasmithC4DTranslator/Private/DatasmithC4DImporter.cpp b/Engine/Plugins/Enterprise/DatasmithC4DImporter/Source/DatasmithC4DTranslator/Private/DatasmithC4DImporter.cpp
--- a/Engine/Plugins/Enterprise/DatasmithC4DImporter/Source/DatasmithC4DTranslator/Private/DatasmithC4DImporter.cpp
+++ b/Engine/Plugins/Enterprise/DatasmithC4DImporter/Source/DatasmithC4DTranslator/Private/DatasmithC4DImporter.cpp
@@ -1216,7 +1216,8 @@ TSharedPtr<IDatasmithTextureElement> FDatasmithC4DImporter::ImportTexture(const
 		return nullptr;
 	}
 
-	FString TextureName = FString::Printf(TEXT("%ls_%d"), *FMD5::HashAnsiString(*TexturePath), int32(TextureMode));
+	const FString TextureHash = FMD5::HashAnsiString(*TexturePath);
+	FString TextureName = FString::Printf(TEXT("%ls_%d"), *TextureHash, int32(TextureMode));
 	if (TSharedPtr<IDatasmithTextureElement>* FoundImportedTexture = ImportedTextures.Find(TextureName))
 	{
 		return *FoundImportedTexture;
@@ -2832,6 +2833,10 @@ TSharedPtr<IDatasmithMeshElement> FDatasmithC4DImporter::ImportMesh(cineware::Po
 
 	// At least one UV set must exist.
 	int32 UVChannelCount = UVWTagsData.Num();
+	if (!VertexInstanceUVs.IsValid())
+	{
+		return TSharedPtr<IDatasmithMeshElement>();
+	}
 	VertexInstanceUVs.SetNumChannels(FMath::Max(1, UVChannelCount));
 
 	// Vertices
@@ -3152,7 +3157,7 @@ cineware::BaseObject* FDatasmithC4DImporter::GoToMelangeHierarchyPosition(cinewa
 			{
 				Object = GoToMelangeHierarchyPosition(GetBestMelangeCache(Object), NextHierarchyPosition.Right(NextHierarchyPosition.Len() - 2));
 			}
-			else
+			else if(Object)
 			{
 				Object = GoToMelangeHierarchyPosition(Object->GetDown(), NextHierarchyPosition);
 			}
