Add test on pointer to avoid  null pointer dereferencing

diff --git a/Engine/Plugins/Enterprise/DatasmithC4DImporter/Source/DatasmithC4DTranslator/Private/DatasmithC4DImporter.cpp b/Engine/Plugins/Enterprise/DatasmithC4DImporter/Source/DatasmithC4DTranslator/Private/DatasmithC4DImporter.cpp
--- a/Engine/Plugins/Enterprise/DatasmithC4DImporter/Source/DatasmithC4DTranslator/Private/DatasmithC4DImporter.cpp
+++ b/Engine/Plugins/Enterprise/DatasmithC4DImporter/Source/DatasmithC4DTranslator/Private/DatasmithC4DImporter.cpp
@@ -1216,8 +1216,14 @@ TSharedPtr<IDatasmithTextureElement> FDatasmithC4DImporter::ImportTexture(const
 		return nullptr;
 	}
 
-	const FString TextureHash = FMD5::HashAnsiString(*TexturePath);
-	FString TextureName = FString::Printf(TEXT("%ls_%d"), *TextureHash, int32(TextureMode));
+	const TCHAR* TexturePathChr = *TexturePath;
+	if (TexturePathChr == nullptr)
+	{
+		return nullptr;
+	}
+
+	const FString TextureHash = FMD5::HashAnsiString(TexturePathChr);
+	const FString TextureName = FString::Printf(TEXT("%ls_%d"), *TextureHash, int32(TextureMode));
 	if (TSharedPtr<IDatasmithTextureElement>* FoundImportedTexture = ImportedTextures.Find(TextureName))
 	{
 		return *FoundImportedTexture;
