Fixed warning on during Static Analysis Mac (Clang) UnrealEditor and UnrealGame

#rnx

diff --git a/Engine/Plugins/Enterprise/DatasmithC4DImporter/Source/DatasmithC4DTranslator/Private/DatasmithC4DImporter.cpp b/Engine/Plugins/Enterprise/DatasmithC4DImporter/Source/DatasmithC4DTranslator/Private/DatasmithC4DImporter.cpp
--- a/Engine/Plugins/Enterprise/DatasmithC4DImporter/Source/DatasmithC4DTranslator/Private/DatasmithC4DImporter.cpp
+++ b/Engine/Plugins/Enterprise/DatasmithC4DImporter/Source/DatasmithC4DTranslator/Private/DatasmithC4DImporter.cpp
@@ -2768,11 +2768,27 @@ void FDatasmithC4DImporter::ImportHierarchy(cineware::BaseObject* ActorObject, c
 
 TSharedPtr<IDatasmithMeshElement> FDatasmithC4DImporter::ImportMesh(cineware::PolygonObject* PolyObject, const FString& DatasmithMeshName, const FString& DatasmithLabel)
 {
+	if (!PolyObject)
+	{
+		return nullptr;
+	}
+
 	cineware::Int32 PointCount = PolyObject->GetPointCount();
 	cineware::Int32 PolygonCount = PolyObject->GetPolygonCount();
+	if (PointCount == 0 || PolygonCount == 0)
+	{
+		// Missing points or triangles. Skipping this polygon object
+		return nullptr;
+	}
 
 	const cineware::Vector* Points = PolyObject->GetPointR();
 	const cineware::CPolygon* Polygons = PolyObject->GetPolygonR();
+	if (!Points || !Polygons)
+	{
+		// Something wrong happened. Skipping this polygon object
+		ensure(false);
+		return nullptr;
+	}
 
 	// Get vertex normals
 	cineware::Vector32* Normals = nullptr;
@@ -2846,11 +2862,20 @@ TSharedPtr<IDatasmithMeshElement> FDatasmithC4DImporter::ImportMesh(cineware::Po
 	VertexInstanceUVs.SetNumChannels(FMath::Max(1, UVChannelCount));
 
 	// Vertices
+	if (!VertexPositions.IsValid())
+	{
+		return nullptr;
+	}
+
 	for (int32 PointIndex = 0; PointIndex < PointCount; ++PointIndex)
 	{
 		FVertexID NewVertexID = MeshDescription.CreateVertex();
 		// We count on this check when creating polygons
-		check(NewVertexID.GetValue() == PointIndex);
+		if (NewVertexID.GetValue() != PointIndex)
+		{
+			// Something wrong happened
+			return nullptr;
+		}
 
 		VertexPositions[NewVertexID] = (FVector3f)ConvertMelangePosition(Points[PointIndex]);
 	}
