Workaround macro conflicting with C4D header files breaking non-unity build

diff --git a/Engine/Plugins/Enterprise/DatasmithC4DImporter/Source/DatasmithC4DTranslator/Private/DatasmithC4DImporter.cpp b/Engine/Plugins/Enterprise/DatasmithC4DImporter/Source/DatasmithC4DTranslator/Private/DatasmithC4DImporter.cpp
--- a/Engine/Plugins/Enterprise/DatasmithC4DImporter/Source/DatasmithC4DTranslator/Private/DatasmithC4DImporter.cpp
+++ b/Engine/Plugins/Enterprise/DatasmithC4DImporter/Source/DatasmithC4DTranslator/Private/DatasmithC4DImporter.cpp
@@ -802,6 +802,11 @@ TSharedPtr<IDatasmithActorElement> FDatasmithC4DImporter::ImportNullActor(cinewa
 	return ActorElement;
 }
 
+// workaround define from UE that conflicts with C4D SDK constant
+#ifdef LIGHT_TYPE_SPOT
+#undef LIGHT_TYPE_SPOT
+#endif
+
 namespace
 {
 	TSharedPtr<IDatasmithLightActorElement> CreateDatasmithLightActorElement(int32 MelangeLightTypeId, const FString& Name, const FString& Label)
