- WinArm64 SymsLib (and updated .sln)

diff --git a/Engine/Plugins/Enterprise/DatasmithC4DImporter/Source/DatasmithC4DTranslator/DatasmithC4DTranslator.Build.cs b/Engine/Plugins/Enterprise/DatasmithC4DImporter/Source/DatasmithC4DTranslator/DatasmithC4DTranslator.Build.cs
--- a/Engine/Plugins/Enterprise/DatasmithC4DImporter/Source/DatasmithC4DTranslator/DatasmithC4DTranslator.Build.cs
+++ b/Engine/Plugins/Enterprise/DatasmithC4DImporter/Source/DatasmithC4DTranslator/DatasmithC4DTranslator.Build.cs
@@ -47,21 +47,24 @@ namespace UnrealBuildTool.Rules
 			);
 
 
-
-			// Set up the C4D Melange SDK includes and libraries.
-			string MelangeSDKLocation = Path.Combine(EngineDirectory, "Restricted/NotForLicensees/Source/ThirdParty/Enterprise/Melange/20.004_RBMelange20.0_259890");
-
-			// When C4D Melange SDK is not part of the developer's workspace, look for environment variable Melange_SDK.
-			if (!Directory.Exists(MelangeSDKLocation))
+			// no winarm64 sdk exists yet
+			if (!(Target.Platform == UnrealTargetPlatform.Win64 && Target.Architecture == UnrealArch.Arm64))
 			{
-				MelangeSDKLocation = System.Environment.GetEnvironmentVariable("Melange_SDK");
-			}
+				// Set up the C4D Melange SDK includes and libraries.
+				string MelangeSDKLocation = Path.Combine(EngineDirectory, "Restricted/NotForLicensees/Source/ThirdParty/Enterprise/Melange/20.004_RBMelange20.0_259890");
 
-			// Make sure the C4D Melange SDK is actually installed.
-			if (Directory.Exists(MelangeSDKLocation))
-			{
-				PublicDefinitions.Add("_MELANGE_SDK_");
-				PrivateDependencyModuleNames.Add("MelangeSDK");
+				// When C4D Melange SDK is not part of the developer's workspace, look for environment variable Melange_SDK.
+				if (!Directory.Exists(MelangeSDKLocation))
+				{
+					MelangeSDKLocation = System.Environment.GetEnvironmentVariable("Melange_SDK");
+				}
+
+				// Make sure the C4D Melange SDK is actually installed.
+				if (Directory.Exists(MelangeSDKLocation))
+				{
+					PublicDefinitions.Add("_MELANGE_SDK_");
+					PrivateDependencyModuleNames.Add("MelangeSDK");
+				}
 			}
 
 			if (Target.WindowsPlatform.Compiler == WindowsCompiler.Clang)
diff --git a/Engine/Source/Runtime/SymsLib/SymsLib.Build.cs b/Engine/Source/Runtime/SymsLib/SymsLib.Build.cs
--- a/Engine/Source/Runtime/SymsLib/SymsLib.Build.cs
+++ b/Engine/Source/Runtime/SymsLib/SymsLib.Build.cs
@@ -17,8 +17,9 @@ public class SymsLib : ModuleRules
 		if (Target.Platform.IsInGroup(UnrealPlatformGroup.Windows))
 		{
 			bool bUseDebug = Target.Configuration == UnrealTargetConfiguration.Debug;
-			string LibPath = LibPathBase + "/x64" + (bUseDebug ? "/Debug" : "/Release");
-			PublicAdditionalLibraries.Add(LibPath + "/SymsLib.lib");
+			string Arch = Target.Architecture == UnrealArch.Arm64 ? "ARM64" : "x64";
+			string LibPath = Path.Combine(LibPathBase, Arch, (bUseDebug ? "Debug" : "Release"), "SymsLib.lib");
+			PublicAdditionalLibraries.Add(LibPath);
 		}
 		else if (Target.Platform.IsInGroup(UnrealPlatformGroup.Unix))
 		{
