Fixed crash of Automotive templates

#rnx

#Synced-CL 50708133

diff --git a/Engine/Plugins/Enterprise/DatasmithC4DImporter/Source/DatasmithC4DTranslator/Private/DatasmithC4DImportOptions.cpp b/Engine/Plugins/Enterprise/DatasmithC4DImporter/Source/DatasmithC4DTranslator/Private/DatasmithC4DImportOptions.cpp
--- a/Engine/Plugins/Enterprise/DatasmithC4DImporter/Source/DatasmithC4DTranslator/Private/DatasmithC4DImportOptions.cpp
+++ b/Engine/Plugins/Enterprise/DatasmithC4DImporter/Source/DatasmithC4DTranslator/Private/DatasmithC4DImportOptions.cpp
@@ -17,19 +17,6 @@ UDatasmithC4DImportOptions::UDatasmithC4DImportOptions(const FObjectInitializer&
 	ScaleVertices = 1.0;
 #if WITH_EDITOR
 	bExportToUDatasmith = false;
-
-	// In debug show all properties including the ones for debug
-	if (IDatasmithC4DTranslatorModule::Get().InDebugMode())
-	{
-		for (TFieldIterator<FProperty> It(GetClass()); It; ++It)
-		{
-			FProperty* Property = *It;
-			if (Property && Property->HasMetaData(TEXT("Category")) && Property->GetMetaData(TEXT("Category")) == TEXT("DebugProperty"))
-			{
-				Property->SetMetaData(TEXT("Category"), TEXT("PrivateSettings"));
-			}
-		}
-	}
 #endif //WITH_EDITOR
 }
 
diff --git a/Engine/Plugins/Enterprise/DatasmithC4DImporter/Source/DatasmithC4DTranslator/Private/DatasmithC4DTranslatorModule.cpp b/Engine/Plugins/Enterprise/DatasmithC4DImporter/Source/DatasmithC4DTranslator/Private/DatasmithC4DTranslatorModule.cpp
--- a/Engine/Plugins/Enterprise/DatasmithC4DImporter/Source/DatasmithC4DTranslator/Private/DatasmithC4DTranslatorModule.cpp
+++ b/Engine/Plugins/Enterprise/DatasmithC4DImporter/Source/DatasmithC4DTranslator/Private/DatasmithC4DTranslatorModule.cpp
@@ -6,6 +6,7 @@
 
 #include "DatasmithC4DTranslator.h"
 #include "DatasmithC4DImporterMaterialSelector.h"
+#include "DatasmithC4DImportOptions.h"
 
 #include "ReferenceMaterials/DatasmithReferenceMaterialManager.h"
 
@@ -40,8 +41,18 @@ public:
 
 		FDatasmithReferenceMaterialManager::Get().RegisterSelector(TEXT("C4DTranslator"), MakeShared< FDatasmithC4DImporterMaterialSelector >());
 
-		FString EnvVariable = FPlatformMisc::GetEnvironmentVariable(TEXT("DATASMITHC4D_DEBUG"));
-		bDebugMode = !EnvVariable.IsEmpty();
+		const FString EnvVariable = FPlatformMisc::GetEnvironmentVariable(TEXT("DATASMITHC4D_DEBUG"));
+		if (!EnvVariable.IsEmpty())
+		{
+			for (TFieldIterator<FProperty> It(UDatasmithC4DImportOptions::StaticClass()); It; ++It)
+			{
+				FProperty* Property = *It;
+				if (Property && Property->HasMetaData(TEXT("Category")) && Property->GetMetaData(TEXT("Category")) == TEXT("DebugProperty"))
+				{
+					Property->SetMetaData(TEXT("Category"), TEXT("PrivateSettings"));
+				}
+			}
+		}
 
 		// TODO: Load C4DDynamicImporter if available else Import Static
 
@@ -50,18 +61,12 @@ public:
 
 	virtual void ShutdownModule() override
 	{
-		Datasmith::UnregisterTranslator<FDatasmithC4DTranslator>();
+		if (FModuleManager::Get().IsModuleLoaded(TEXT("DatasmithTranslator")))
+		{
+			Datasmith::UnregisterTranslator<FDatasmithC4DTranslator>();
+			FDatasmithReferenceMaterialManager::Get().UnregisterSelector(TEXT("C4DTranslator"));
+		}
 	}
-
-	bool InDebugMode() const override
-	{
-		return bDebugMode;
-	}
-
-private:
-
-	/** Indicates if DATASMITHC4D_DEBUG environment variable is set */
-	bool bDebugMode;
 };
 
 #else //_MELANGE_SDK_
@@ -78,11 +83,6 @@ public:
 	virtual void ShutdownModule() override
 	{
 	}
-
-	bool InDebugMode() const override
-	{
-		return false;
-	}
 };
 #endif //_MELANGE_SDK_
 
diff --git a/Engine/Plugins/Enterprise/DatasmithC4DImporter/Source/DatasmithC4DTranslator/Public/DatasmithC4DTranslatorModule.h b/Engine/Plugins/Enterprise/DatasmithC4DImporter/Source/DatasmithC4DTranslator/Public/DatasmithC4DTranslatorModule.h
--- a/Engine/Plugins/Enterprise/DatasmithC4DImporter/Source/DatasmithC4DTranslator/Public/DatasmithC4DTranslatorModule.h
+++ b/Engine/Plugins/Enterprise/DatasmithC4DImporter/Source/DatasmithC4DTranslator/Public/DatasmithC4DTranslatorModule.h
@@ -20,9 +20,4 @@ public:
 	{
 		return FModuleManager::Get().IsModuleLoaded(C4DTRANSLATOR_MODULE_NAME);
 	}
-
-	/**
-	 * @returns true if environment variable DATASMITHC4D_DEBUG is set
-	 */
-	virtual bool InDebugMode() const = 0;
 };
diff --git a/Engine/Plugins/Enterprise/DatasmithCADImporter/Config/DefaultDatasmithCADImporter.ini b/Engine/Plugins/Enterprise/DatasmithCADImporter/Config/DefaultDatasmithCADImporter.ini
--- a/Engine/Plugins/Enterprise/DatasmithCADImporter/Config/DefaultDatasmithCADImporter.ini
+++ b/Engine/Plugins/Enterprise/DatasmithCADImporter/Config/DefaultDatasmithCADImporter.ini
@@ -13,6 +13,8 @@
 
 +ClassRedirects=(OldName="/Script/ParametricSurface.UParametricSurfaceData", NewName="/Script/ParametricSurface.UDatasmithParametricSurfaceData")
 
++PropertyRedirects=(OldName="UDatasmithParametricSurfaceData.RawData",NewName="RawData_DEPRECATED")
+
 [SectionsToSave]
 +Section=/Script/DatasmithWireTranslator.DatasmithWireOptions
 +Section=/Script/DatasmithOpenNurbsTranslator.DatasmithOpenNurbsImportOptions
diff --git a/Engine/Plugins/Enterprise/DatasmithCADImporter/Source/ParametricSurface/Private/ParametricSurfaceModule.cpp b/Engine/Plugins/Enterprise/DatasmithCADImporter/Source/ParametricSurface/Private/ParametricSurfaceModule.cpp
--- a/Engine/Plugins/Enterprise/DatasmithCADImporter/Source/ParametricSurface/Private/ParametricSurfaceModule.cpp
+++ b/Engine/Plugins/Enterprise/DatasmithCADImporter/Source/ParametricSurface/Private/ParametricSurfaceModule.cpp
@@ -12,15 +12,6 @@
 
 void FParametricSurfaceModule::StartupModule()
 {
-	{
-		TArray<FCoreRedirect> Redirects;
-		Redirects.Emplace(ECoreRedirectFlags::Type_Package, TEXT("/Script/DatasmithCoreTechParametricSurfaceData"), TEXT("/Script/ParametricSurface"));
-		Redirects.Emplace(ECoreRedirectFlags::Type_Property, TEXT("UDatasmithParametricSurfaceData.RawData"), TEXT("RawData_DEPRECATED"));
-		Redirects.Emplace(ECoreRedirectFlags::Type_Property, TEXT("UCoreTechParametricSurfaceData.RawData"), TEXT("RawData_DEPRECATED"));
-		Redirects.Emplace(ECoreRedirectFlags::Type_Struct, TEXT("CoreTechSceneParameters"), TEXT("/Script/ParametricSurface.ParametricSceneParameters"));
-		Redirects.Emplace(ECoreRedirectFlags::Type_Struct, TEXT("CoreTechMeshParameters"), TEXT("/Script/ParametricSurface.ParametricMeshParameters"));
-		FCoreRedirects::AddRedirectList(Redirects, PARAMETRICSURFACE_MODULE_NAME);
-	}
 }
 
 FParametricSurfaceModule& FParametricSurfaceModule::Get()
diff --git a/Engine/Plugins/Enterprise/DatasmithImporter/Source/DatasmithTranslator/Private/DatasmithTranslatorModule.cpp b/Engine/Plugins/Enterprise/DatasmithImporter/Source/DatasmithTranslator/Private/DatasmithTranslatorModule.cpp
--- a/Engine/Plugins/Enterprise/DatasmithImporter/Source/DatasmithTranslator/Private/DatasmithTranslatorModule.cpp
+++ b/Engine/Plugins/Enterprise/DatasmithImporter/Source/DatasmithTranslator/Private/DatasmithTranslatorModule.cpp
@@ -2,7 +2,6 @@
 
 #include "DatasmithTranslatorModule.h"
 
-#include "ReferenceMaterials/DatasmithC4DMaterialSelector.h"
 #include "ReferenceMaterials/DatasmithCityEngineMaterialSelector.h"
 #include "ReferenceMaterials/DatasmithReferenceMaterialManager.h"
 #include "ReferenceMaterials/DatasmithRevitMaterialSelector.h"
@@ -14,7 +13,6 @@ void IDatasmithTranslatorModule::StartupModule()
 	FDatasmithReferenceMaterialManager::Create();
 
 	//A minimal set of natively supported reference materials
-	FDatasmithReferenceMaterialManager::Get().RegisterSelector(TEXT("C4D"), MakeShared< FDatasmithC4DMaterialSelector >());
 	FDatasmithReferenceMaterialManager::Get().RegisterSelector(TEXT("Revit"), MakeShared< FDatasmithRevitMaterialSelector >());
 	FDatasmithReferenceMaterialManager::Get().RegisterSelector(TEXT("SketchUp"), MakeShared< FDatasmithSketchUpMaterialSelector >());
 	FDatasmithReferenceMaterialManager::Get().RegisterSelector(TEXT("CityEngine"), MakeShared< FDatasmithCityEngineMaterialSelector >());
@@ -23,7 +21,6 @@ void IDatasmithTranslatorModule::StartupModule()
 
 void IDatasmithTranslatorModule::ShutdownModule()
 {
-	FDatasmithReferenceMaterialManager::Get().UnregisterSelector(TEXT("C4D"));
 	FDatasmithReferenceMaterialManager::Get().UnregisterSelector(TEXT("Revit"));
 	FDatasmithReferenceMaterialManager::Get().UnregisterSelector(TEXT("SketchUp"));
 	FDatasmithReferenceMaterialManager::Get().UnregisterSelector(TEXT("CityEngine"));
diff --git a/Engine/Plugins/Enterprise/DatasmithImporter/Source/DatasmithTranslator/Private/ReferenceMaterials/DatasmithC4DMaterialSelector.cpp b/Engine/Plugins/Enterprise/DatasmithImporter/Source/DatasmithTranslator/Private/ReferenceMaterials/DatasmithC4DMaterialSelector.cpp
deleted file mode 100644
diff --git a/Engine/Plugins/Enterprise/DatasmithImporter/Source/DatasmithTranslator/Private/ReferenceMaterials/DatasmithC4DMaterialSelector.h b/Engine/Plugins/Enterprise/DatasmithImporter/Source/DatasmithTranslator/Private/ReferenceMaterials/DatasmithC4DMaterialSelector.h
deleted file mode 100644
diff --git a/Engine/Plugins/Enterprise/DatasmithImporter/Source/DatasmithTranslator/Private/ReferenceMaterials/DatasmithReferenceMaterialManager.cpp b/Engine/Plugins/Enterprise/DatasmithImporter/Source/DatasmithTranslator/Private/ReferenceMaterials/DatasmithReferenceMaterialManager.cpp
--- a/Engine/Plugins/Enterprise/DatasmithImporter/Source/DatasmithTranslator/Private/ReferenceMaterials/DatasmithReferenceMaterialManager.cpp
+++ b/Engine/Plugins/Enterprise/DatasmithImporter/Source/DatasmithTranslator/Private/ReferenceMaterials/DatasmithReferenceMaterialManager.cpp
@@ -6,6 +6,11 @@
 
 TUniquePtr< FDatasmithReferenceMaterialManager > FDatasmithReferenceMaterialManager::Instance;
 
+FDatasmithReferenceMaterialManager::~FDatasmithReferenceMaterialManager()
+{
+	Selectors.Empty();
+}
+
 void FDatasmithReferenceMaterialManager::Create()
 {
 	Instance = MakeUnique< FDatasmithReferenceMaterialManager >();
diff --git a/Engine/Plugins/Enterprise/DatasmithImporter/Source/DatasmithTranslator/Public/ReferenceMaterials/DatasmithReferenceMaterialManager.h b/Engine/Plugins/Enterprise/DatasmithImporter/Source/DatasmithTranslator/Public/ReferenceMaterials/DatasmithReferenceMaterialManager.h
--- a/Engine/Plugins/Enterprise/DatasmithImporter/Source/DatasmithTranslator/Public/ReferenceMaterials/DatasmithReferenceMaterialManager.h
+++ b/Engine/Plugins/Enterprise/DatasmithImporter/Source/DatasmithTranslator/Public/ReferenceMaterials/DatasmithReferenceMaterialManager.h
@@ -10,6 +10,9 @@ class FDatasmithReferenceMaterialSelector;
 class DATASMITHTRANSLATOR_API FDatasmithReferenceMaterialManager
 {
 public:
+	FDatasmithReferenceMaterialManager() = default;
+	~FDatasmithReferenceMaterialManager();
+
 	static void Create();
 	static void Destroy();
 	static FDatasmithReferenceMaterialManager& Get();
