[Backout] - CL50739651
#rnx
Original CL Desc
-----------------------------------------------------------------
Fixed crash of Automotive templates

#rnx

diff --git a/Engine/Plugins/Enterprise/DatasmithC4DImporter/Source/DatasmithC4DTranslator/Private/DatasmithC4DImportOptions.cpp b/Engine/Plugins/Enterprise/DatasmithC4DImporter/Source/DatasmithC4DTranslator/Private/DatasmithC4DImportOptions.cpp
--- a/Engine/Plugins/Enterprise/DatasmithC4DImporter/Source/DatasmithC4DTranslator/Private/DatasmithC4DImportOptions.cpp
+++ b/Engine/Plugins/Enterprise/DatasmithC4DImporter/Source/DatasmithC4DTranslator/Private/DatasmithC4DImportOptions.cpp
@@ -17,6 +17,19 @@ UDatasmithC4DImportOptions::UDatasmithC4DImportOptions(const FObjectInitializer&
 	ScaleVertices = 1.0;
 #if WITH_EDITOR
 	bExportToUDatasmith = false;
+
+	// In debug show all properties including the ones for debug
+	if (IDatasmithC4DTranslatorModule::Get().InDebugMode())
+	{
+		for (TFieldIterator<FProperty> It(GetClass()); It; ++It)
+		{
+			FProperty* Property = *It;
+			if (Property && Property->HasMetaData(TEXT("Category")) && Property->GetMetaData(TEXT("Category")) == TEXT("DebugProperty"))
+			{
+				Property->SetMetaData(TEXT("Category"), TEXT("PrivateSettings"));
+			}
+		}
+	}
 #endif //WITH_EDITOR
 }
 
diff --git a/Engine/Plugins/Enterprise/DatasmithC4DImporter/Source/DatasmithC4DTranslator/Private/DatasmithC4DTranslatorModule.cpp b/Engine/Plugins/Enterprise/DatasmithC4DImporter/Source/DatasmithC4DTranslator/Private/DatasmithC4DTranslatorModule.cpp
--- a/Engine/Plugins/Enterprise/DatasmithC4DImporter/Source/DatasmithC4DTranslator/Private/DatasmithC4DTranslatorModule.cpp
+++ b/Engine/Plugins/Enterprise/DatasmithC4DImporter/Source/DatasmithC4DTranslator/Private/DatasmithC4DTranslatorModule.cpp
@@ -6,7 +6,6 @@
 
 #include "DatasmithC4DTranslator.h"
 #include "DatasmithC4DImporterMaterialSelector.h"
-#include "DatasmithC4DImportOptions.h"
 
 #include "ReferenceMaterials/DatasmithReferenceMaterialManager.h"
 
@@ -41,18 +40,8 @@ public:
 
 		FDatasmithReferenceMaterialManager::Get().RegisterSelector(TEXT("C4DTranslator"), MakeShared< FDatasmithC4DImporterMaterialSelector >());
 
-		const FString EnvVariable = FPlatformMisc::GetEnvironmentVariable(TEXT("DATASMITHC4D_DEBUG"));
-		if (!EnvVariable.IsEmpty())
-		{
-			for (TFieldIterator<FProperty> It(UDatasmithC4DImportOptions::StaticClass()); It; ++It)
-			{
-				FProperty* Property = *It;
-				if (Property && Property->HasMetaData(TEXT("Category")) && Property->GetMetaData(TEXT("Category")) == TEXT("DebugProperty"))
-				{
-					Property->SetMetaData(TEXT("Category"), TEXT("PrivateSettings"));
-				}
-			}
-		}
+		FString EnvVariable = FPlatformMisc::GetEnvironmentVariable(TEXT("DATASMITHC4D_DEBUG"));
+		bDebugMode = !EnvVariable.IsEmpty();
 
 		// TODO: Load C4DDynamicImporter if available else Import Static
 
@@ -61,12 +50,18 @@ public:
 
 	virtual void ShutdownModule() override
 	{
-		if (FModuleManager::Get().IsModuleLoaded(TEXT("DatasmithTranslator")))
-		{
-			Datasmith::UnregisterTranslator<FDatasmithC4DTranslator>();
-			FDatasmithReferenceMaterialManager::Get().UnregisterSelector(TEXT("C4DTranslator"));
-		}
+		Datasmith::UnregisterTranslator<FDatasmithC4DTranslator>();
 	}
+
+	bool InDebugMode() const override
+	{
+		return bDebugMode;
+	}
+
+private:
+
+	/** Indicates if DATASMITHC4D_DEBUG environment variable is set */
+	bool bDebugMode;
 };
 
 #else //_MELANGE_SDK_
@@ -83,6 +78,11 @@ public:
 	virtual void ShutdownModule() override
 	{
 	}
+
+	bool InDebugMode() const override
+	{
+		return false;
+	}
 };
 #endif //_MELANGE_SDK_
 
diff --git a/Engine/Plugins/Enterprise/DatasmithC4DImporter/Source/DatasmithC4DTranslator/Public/DatasmithC4DTranslatorModule.h b/Engine/Plugins/Enterprise/DatasmithC4DImporter/Source/DatasmithC4DTranslator/Public/DatasmithC4DTranslatorModule.h
--- a/Engine/Plugins/Enterprise/DatasmithC4DImporter/Source/DatasmithC4DTranslator/Public/DatasmithC4DTranslatorModule.h
+++ b/Engine/Plugins/Enterprise/DatasmithC4DImporter/Source/DatasmithC4DTranslator/Public/DatasmithC4DTranslatorModule.h
@@ -20,4 +20,9 @@ public:
 	{
 		return FModuleManager::Get().IsModuleLoaded(C4DTRANSLATOR_MODULE_NAME);
 	}
+
+	/**
+	 * @returns true if environment variable DATASMITHC4D_DEBUG is set
+	 */
+	virtual bool InDebugMode() const = 0;
 };
diff --git a/Engine/Plugins/Enterprise/DatasmithCADImporter/Config/DefaultDatasmithCADImporter.ini b/Engine/Plugins/Enterprise/DatasmithCADImporter/Config/DefaultDatasmithCADImporter.ini
--- a/Engine/Plugins/Enterprise/DatasmithCADImporter/Config/DefaultDatasmithCADImporter.ini
+++ b/Engine/Plugins/Enterprise/DatasmithCADImporter/Config/DefaultDatasmithCADImporter.ini
@@ -13,8 +13,6 @@
 
 +ClassRedirects=(OldName="/Script/ParametricSurface.UParametricSurfaceData", NewName="/Script/ParametricSurface.UDatasmithParametricSurfaceData")
 
-+PropertyRedirects=(OldName="UDatasmithParametricSurfaceData.RawData",NewName="RawData_DEPRECATED")
-
 [SectionsToSave]
 +Section=/Script/DatasmithWireTranslator.DatasmithWireOptions
 +Section=/Script/DatasmithOpenNurbsTranslator.DatasmithOpenNurbsImportOptions
diff --git a/Engine/Plugins/Enterprise/DatasmithCADImporter/Source/ParametricSurface/Private/ParametricSurfaceModule.cpp b/Engine/Plugins/Enterprise/DatasmithCADImporter/Source/ParametricSurface/Private/ParametricSurfaceModule.cpp
--- a/Engine/Plugins/Enterprise/DatasmithCADImporter/Source/ParametricSurface/Private/ParametricSurfaceModule.cpp
+++ b/Engine/Plugins/Enterprise/DatasmithCADImporter/Source/ParametricSurface/Private/ParametricSurfaceModule.cpp
@@ -12,6 +12,15 @@
 
 void FParametricSurfaceModule::StartupModule()
 {
+	{
+		TArray<FCoreRedirect> Redirects;
+		Redirects.Emplace(ECoreRedirectFlags::Type_Package, TEXT("/Script/DatasmithCoreTechParametricSurfaceData"), TEXT("/Script/ParametricSurface"));
+		Redirects.Emplace(ECoreRedirectFlags::Type_Property, TEXT("UDatasmithParametricSurfaceData.RawData"), TEXT("RawData_DEPRECATED"));
+		Redirects.Emplace(ECoreRedirectFlags::Type_Property, TEXT("UCoreTechParametricSurfaceData.RawData"), TEXT("RawData_DEPRECATED"));
+		Redirects.Emplace(ECoreRedirectFlags::Type_Struct, TEXT("CoreTechSceneParameters"), TEXT("/Script/ParametricSurface.ParametricSceneParameters"));
+		Redirects.Emplace(ECoreRedirectFlags::Type_Struct, TEXT("CoreTechMeshParameters"), TEXT("/Script/ParametricSurface.ParametricMeshParameters"));
+		FCoreRedirects::AddRedirectList(Redirects, PARAMETRICSURFACE_MODULE_NAME);
+	}
 }
 
 FParametricSurfaceModule& FParametricSurfaceModule::Get()
diff --git a/Engine/Plugins/Enterprise/DatasmithImporter/Source/DatasmithTranslator/Private/DatasmithTranslatorModule.cpp b/Engine/Plugins/Enterprise/DatasmithImporter/Source/DatasmithTranslator/Private/DatasmithTranslatorModule.cpp
--- a/Engine/Plugins/Enterprise/DatasmithImporter/Source/DatasmithTranslator/Private/DatasmithTranslatorModule.cpp
+++ b/Engine/Plugins/Enterprise/DatasmithImporter/Source/DatasmithTranslator/Private/DatasmithTranslatorModule.cpp
@@ -2,6 +2,7 @@
 
 #include "DatasmithTranslatorModule.h"
 
+#include "ReferenceMaterials/DatasmithC4DMaterialSelector.h"
 #include "ReferenceMaterials/DatasmithCityEngineMaterialSelector.h"
 #include "ReferenceMaterials/DatasmithReferenceMaterialManager.h"
 #include "ReferenceMaterials/DatasmithRevitMaterialSelector.h"
@@ -13,6 +14,7 @@ void IDatasmithTranslatorModule::StartupModule()
 	FDatasmithReferenceMaterialManager::Create();
 
 	//A minimal set of natively supported reference materials
+	FDatasmithReferenceMaterialManager::Get().RegisterSelector(TEXT("C4D"), MakeShared< FDatasmithC4DMaterialSelector >());
 	FDatasmithReferenceMaterialManager::Get().RegisterSelector(TEXT("Revit"), MakeShared< FDatasmithRevitMaterialSelector >());
 	FDatasmithReferenceMaterialManager::Get().RegisterSelector(TEXT("SketchUp"), MakeShared< FDatasmithSketchUpMaterialSelector >());
 	FDatasmithReferenceMaterialManager::Get().RegisterSelector(TEXT("CityEngine"), MakeShared< FDatasmithCityEngineMaterialSelector >());
@@ -21,6 +23,7 @@ void IDatasmithTranslatorModule::StartupModule()
 
 void IDatasmithTranslatorModule::ShutdownModule()
 {
+	FDatasmithReferenceMaterialManager::Get().UnregisterSelector(TEXT("C4D"));
 	FDatasmithReferenceMaterialManager::Get().UnregisterSelector(TEXT("Revit"));
 	FDatasmithReferenceMaterialManager::Get().UnregisterSelector(TEXT("SketchUp"));
 	FDatasmithReferenceMaterialManager::Get().UnregisterSelector(TEXT("CityEngine"));
diff --git a/Engine/Plugins/Enterprise/DatasmithImporter/Source/DatasmithTranslator/Private/ReferenceMaterials/DatasmithC4DMaterialSelector.cpp b/Engine/Plugins/Enterprise/DatasmithImporter/Source/DatasmithTranslator/Private/ReferenceMaterials/DatasmithC4DMaterialSelector.cpp
new file mode 100644
--- /dev/null
+++ b/Engine/Plugins/Enterprise/DatasmithImporter/Source/DatasmithTranslator/Private/ReferenceMaterials/DatasmithC4DMaterialSelector.cpp
@@ -0,0 +1,46 @@
+// Copyright Epic Games, Inc. All Rights Reserved.
+
+#include "DatasmithC4DMaterialSelector.h"
+
+#include "IDatasmithSceneElements.h"
+
+#include "Materials/Material.h"
+#include "Materials/MaterialInstanceConstant.h"
+#include "UObject/SoftObjectPath.h"
+#include "Templates/Casts.h"
+
+FDatasmithC4DMaterialSelector::FDatasmithC4DMaterialSelector()
+{
+	ReferenceMaterial.FromSoftObjectPath( FSoftObjectPath("/DatasmithContent/Materials/C4DReference.C4DReference") );
+}
+
+bool FDatasmithC4DMaterialSelector::IsValid() const
+{
+	return ReferenceMaterial.IsValid();
+}
+
+const FDatasmithReferenceMaterial& FDatasmithC4DMaterialSelector::GetReferenceMaterial( const TSharedPtr< IDatasmithMaterialInstanceElement >& /*InDatasmithMaterial*/ ) const
+{
+	return ReferenceMaterial;
+}
+
+void FDatasmithC4DMaterialSelector::FinalizeMaterialInstance(const TSharedPtr<IDatasmithMaterialInstanceElement>& InDatasmithMaterial, UMaterialInstanceConstant * MaterialInstance) const
+{
+	if (!InDatasmithMaterial.IsValid() || MaterialInstance == nullptr)
+	{
+		return;
+	}
+
+	// Set blend mode to translucent if material requires transparency.
+	if (InDatasmithMaterial->GetMaterialType() == EDatasmithReferenceMaterialType::Transparent)
+	{
+		MaterialInstance->BasePropertyOverrides.bOverride_BlendMode = true;
+		MaterialInstance->BasePropertyOverrides.BlendMode = BLEND_Translucent;
+	}
+	// Set blend mode to masked if material has cutouts.
+	else if (InDatasmithMaterial->GetMaterialType() == EDatasmithReferenceMaterialType::CutOut)
+	{
+		MaterialInstance->BasePropertyOverrides.bOverride_BlendMode = true;
+		MaterialInstance->BasePropertyOverrides.BlendMode = BLEND_Masked;
+	}
+}
diff --git a/Engine/Plugins/Enterprise/DatasmithImporter/Source/DatasmithTranslator/Private/ReferenceMaterials/DatasmithC4DMaterialSelector.h b/Engine/Plugins/Enterprise/DatasmithImporter/Source/DatasmithTranslator/Private/ReferenceMaterials/DatasmithC4DMaterialSelector.h
new file mode 100644
--- /dev/null
+++ b/Engine/Plugins/Enterprise/DatasmithImporter/Source/DatasmithTranslator/Private/ReferenceMaterials/DatasmithC4DMaterialSelector.h
@@ -0,0 +1,27 @@
+// Copyright Epic Games, Inc. All Rights Reserved.
+
+#pragma once
+
+#include "DatasmithDefinitions.h"
+#include "ReferenceMaterials/DatasmithReferenceMaterial.h"
+#include "ReferenceMaterials/DatasmithReferenceMaterialSelector.h"
+
+#include "Templates/SharedPointer.h"
+
+class IDatasmithMaterialInstanceElement;
+
+class FDatasmithC4DMaterialSelector : public FDatasmithReferenceMaterialSelector
+{
+public:
+
+	FDatasmithC4DMaterialSelector();
+	virtual ~FDatasmithC4DMaterialSelector() = default;
+
+	virtual bool IsValid() const override;
+	virtual const FDatasmithReferenceMaterial& GetReferenceMaterial( const TSharedPtr< IDatasmithMaterialInstanceElement >& /*InDatasmithMaterial*/ ) const override;
+	virtual void FinalizeMaterialInstance(const TSharedPtr< IDatasmithMaterialInstanceElement >& InDatasmithMaterial, UMaterialInstanceConstant* MaterialInstance) const override;
+
+private:
+
+	FDatasmithReferenceMaterial ReferenceMaterial;
+};
diff --git a/Engine/Plugins/Enterprise/DatasmithImporter/Source/DatasmithTranslator/Private/ReferenceMaterials/DatasmithReferenceMaterialManager.cpp b/Engine/Plugins/Enterprise/DatasmithImporter/Source/DatasmithTranslator/Private/ReferenceMaterials/DatasmithReferenceMaterialManager.cpp
--- a/Engine/Plugins/Enterprise/DatasmithImporter/Source/DatasmithTranslator/Private/ReferenceMaterials/DatasmithReferenceMaterialManager.cpp
+++ b/Engine/Plugins/Enterprise/DatasmithImporter/Source/DatasmithTranslator/Private/ReferenceMaterials/DatasmithReferenceMaterialManager.cpp
@@ -6,11 +6,6 @@
 
 TUniquePtr< FDatasmithReferenceMaterialManager > FDatasmithReferenceMaterialManager::Instance;
 
-FDatasmithReferenceMaterialManager::~FDatasmithReferenceMaterialManager()
-{
-	Selectors.Empty();
-}
-
 void FDatasmithReferenceMaterialManager::Create()
 {
 	Instance = MakeUnique< FDatasmithReferenceMaterialManager >();
diff --git a/Engine/Plugins/Enterprise/DatasmithImporter/Source/DatasmithTranslator/Public/ReferenceMaterials/DatasmithReferenceMaterialManager.h b/Engine/Plugins/Enterprise/DatasmithImporter/Source/DatasmithTranslator/Public/ReferenceMaterials/DatasmithReferenceMaterialManager.h
--- a/Engine/Plugins/Enterprise/DatasmithImporter/Source/DatasmithTranslator/Public/ReferenceMaterials/DatasmithReferenceMaterialManager.h
+++ b/Engine/Plugins/Enterprise/DatasmithImporter/Source/DatasmithTranslator/Public/ReferenceMaterials/DatasmithReferenceMaterialManager.h
@@ -10,9 +10,6 @@ class FDatasmithReferenceMaterialSelector;
 class DATASMITHTRANSLATOR_API FDatasmithReferenceMaterialManager
 {
 public:
-	FDatasmithReferenceMaterialManager() = default;
-	~FDatasmithReferenceMaterialManager();
-
 	static void Create();
 	static void Destroy();
 	static FDatasmithReferenceMaterialManager& Get();
